UE-74862 : Reusable media encoder framework implementation

diff --git a/Engine/Plugins/Media/PixelStreaming/Source/PixelStreaming/PixelStreaming.Build.cs b/Engine/Plugins/Media/PixelStreaming/Source/PixelStreaming/PixelStreaming.Build.cs
--- a/Engine/Plugins/Media/PixelStreaming/Source/PixelStreaming/PixelStreaming.Build.cs
+++ b/Engine/Plugins/Media/PixelStreaming/Source/PixelStreaming/PixelStreaming.Build.cs
@@ -112,6 +112,7 @@ namespace UnrealBuildTool.Rules
 				"WebSockets",
 				"Sockets",
                 "MediaUtils",
+				"AVEncoder",
             });
 
             DynamicallyLoadedModuleNames.AddRange(new string[]
@@ -171,6 +172,7 @@ namespace UnrealBuildTool.Rules
             AddSignallingServer();
     	    AddMatchmakingServer();
         	AddWebRTCServers();
+
         }
     }
 }
diff --git a/Engine/Plugins/Media/PixelStreaming/Source/PixelStreaming/Private/Codecs/VideoDecoder.cpp b/Engine/Plugins/Media/PixelStreaming/Source/PixelStreaming/Private/Codecs/VideoDecoder.cpp
--- a/Engine/Plugins/Media/PixelStreaming/Source/PixelStreaming/Private/Codecs/VideoDecoder.cpp
+++ b/Engine/Plugins/Media/PixelStreaming/Source/PixelStreaming/Private/Codecs/VideoDecoder.cpp
@@ -321,13 +321,7 @@ bool FVideoDecoder::QueueBuffer(const webrtc::EncodedImage& InputImage, bool Mis
 
 	TRefCountPtr<IMFMediaBuffer> MediaBuffer;
 
-	int64 CaptureTs = 0;
 	SIZE_T BufferSize = InputImage._length;
-	if (FHUDStats::Get().bEnabled)
-	{
-		BufferSize -= sizeof(CaptureTs); // capture timestamp is appended to encoded frame
-		CaptureTs = *reinterpret_cast<const int64*>(InputImage._buffer + BufferSize);
-	}
 
 	CHECK_HR(MFCreateMemoryBuffer(BufferSize, MediaBuffer.GetInitReference()));
 
@@ -345,10 +339,15 @@ bool FVideoDecoder::QueueBuffer(const webrtc::EncodedImage& InputImage, bool Mis
 	CHECK_HR(Sample->AddBuffer(MediaBuffer));
 	// don't bother converting 90KHz -> 10MHz, decoder doesn't care and we can lose precision on convertion back and forth
 	CHECK_HR(Sample->SetSampleTime(InputImage.Timestamp()));
+
+	// #AVENCODER : What should we do here with the CaptureTs ? That was part of the buffer, and I removed it.
+#if 0
 	// to pass capture timestamp through decoder we set it as sample duration as we don't use duration
 	CHECK_HR(Sample->SetSampleDuration(CaptureTs));
 
 	UE_LOG(LogVideoDecoder, VeryVerbose, TEXT("(%d) enqueueing sample ts %u, capture ts %lld, queue size %d"), RtcTimeMs(), InputImage.Timestamp(), CaptureTs, InputQueueSize.GetValue() + 1);
+#endif
+
 	verify(InputQueue.Enqueue(Sample));
 	InputQueueSize.Increment();
 	InputQueuedEvent->Trigger();
diff --git a/Engine/Plugins/Media/PixelStreaming/Source/PixelStreaming/Private/Codecs/VideoEncoder.cpp b/Engine/Plugins/Media/PixelStreaming/Source/PixelStreaming/Private/Codecs/VideoEncoder.cpp
--- a/Engine/Plugins/Media/PixelStreaming/Source/PixelStreaming/Private/Codecs/VideoEncoder.cpp
+++ b/Engine/Plugins/Media/PixelStreaming/Source/PixelStreaming/Private/Codecs/VideoEncoder.cpp
@@ -2,15 +2,70 @@
 
 #include "VideoEncoder.h"
 #include "RawFrameBuffer.h"
-#include "Codecs/PixelStreamingBaseVideoEncoder.h"
 #include "PlayerSession.h"
 #include "Utils.h"
+#include "HUDStats.h"
+
+TAutoConsoleVariable<int32> CVarPixelStreamingEncoderMaxBitrate(
+	TEXT("PixelStreaming.Encoder.MaxBitrate"),
+	50000000,
+	TEXT("Max bitrate no matter what WebRTC says, in Bps"),
+	ECVF_RenderThreadSafe);
+
+TAutoConsoleVariable<int32> CVarPixelStreamingEncoderUseBackBufferSize(
+	TEXT("PixelStreaming.Encoder.UseBackBufferSize"),
+	1,
+	TEXT("Whether to use back buffer size or custom size"),
+	ECVF_Cheat);
+
+TAutoConsoleVariable<FString> CVarPixelStreamingEncoderTargetSize(
+	TEXT("PixelStreaming.Encoder.TargetSize"),
+	TEXT("1920x1080"),
+	TEXT("Encoder target size in format widthxheight"),
+	ECVF_Cheat);
+
+TAutoConsoleVariable<int32> CVarPixelStreamingEncoderPrioritiseQuality(
+	TEXT("PixelStreaming.Encoder.PrioritiseQuality"),
+	0,
+	TEXT("Reduces framerate automatically on bitrate reduction to trade FPS/latency for video quality"),
+	ECVF_Cheat);
+
+TAutoConsoleVariable<int32> CVarPixelStreamingEncoderLowBitrate(
+	TEXT("PixelStreaming.Encoder.LowBitrate"),
+	1000000,
+	TEXT("Lower bound of bitrate for quality adaptation, in Bps"),
+	ECVF_Default);
+
+TAutoConsoleVariable<int32> CVarPixelStreamingEncoderHighBitrate(
+	TEXT("PixelStreaming.Encoder.HighBitrate"),
+	5000000,
+	TEXT("Upper bound of bitrate for quality adaptation, in Bps"),
+	ECVF_Default);
+
+TAutoConsoleVariable<float> CVarPixelStreamingEncoderMinFPS(
+	TEXT("PixelStreaming.Encoder.MinFPS"),
+	10.0,
+	TEXT("Minimal FPS for quality adaptation"),
+	ECVF_Default);
+
+TAutoConsoleVariable<int32> CVarPixelStreamingEncoderMinQP(
+	TEXT("PixelStreaming.Encoder.MinQP"),
+	20,
+	TEXT("0-51, lower values result in better quality but higher bitrate"),
+	ECVF_Default);
+
+TAutoConsoleVariable<FString> CVarPixelStreamingEncoderRateControl(
+	TEXT("PixelStreaming.Encoder.RateControl"),
+	TEXT("CBR"),
+	TEXT("PixelStreaming video encoder RateControl mode. Supported modes are `ConstQP`, `VBR`, `CBR`"),
+	ECVF_Default);
+
 
 //
 // FVideoEncoderFactory
 //
 
-FVideoEncoderFactory::FVideoEncoderFactory(FPixelStreamingBaseVideoEncoder& HWEncoder): HWEncoder(HWEncoder)
+FVideoEncoderFactory::FVideoEncoderFactory(FHWEncoderDetails& InHWEncoderDetails): HWEncoderDetails(InHWEncoderDetails)
 {}
 
 void FVideoEncoderFactory::AddSession(FPlayerSession& PlayerSession)
@@ -42,7 +97,7 @@ std::unique_ptr<webrtc::VideoEncoder> FVideoEncoderFactory::CreateVideoEncoder(c
 	bool res = PendingPlayerSessions.Dequeue(Session);
 	checkf(res, TEXT("no player session associated with encoder instance"));
 
-	auto VideoEncoder = std::make_unique<FVideoEncoder>(HWEncoder, *Session);
+	auto VideoEncoder = std::make_unique<FVideoEncoder>(HWEncoderDetails, *Session);
 	Session->SetVideoEncoder(VideoEncoder.get());
 	return VideoEncoder;
 }
@@ -51,8 +106,8 @@ std::unique_ptr<webrtc::VideoEncoder> FVideoEncoderFactory::CreateVideoEncoder(c
 // FVideoEncoder
 //
 
-FVideoEncoder::FVideoEncoder(FPixelStreamingBaseVideoEncoder& HWEncoder, FPlayerSession& InPlayerSession):
-	HWEncoder(HWEncoder),
+FVideoEncoder::FVideoEncoder(FHWEncoderDetails& InHWEncoderDetails, FPlayerSession& InPlayerSession):
+	HWEncoderDetails(InHWEncoderDetails),
 	PlayerSession(&InPlayerSession)
 {
 	check(PlayerSession);
@@ -89,13 +144,13 @@ int32 FVideoEncoder::InitEncode(const webrtc::VideoCodec* CodecSettings, int32 N
 int32 FVideoEncoder::RegisterEncodeCompleteCallback(webrtc::EncodedImageCallback* InCallback)
 {
 	Callback = InCallback;
-	HWEncoder.SubscribeToFrameEncodedEvent(*this);
+	HWEncoderDetails.Encoder->RegisterListener(*this);
 	return 0;
 }
 
 int32 FVideoEncoder::Release()
 {
-	HWEncoder.UnsubscribeFromFrameEncodedEvent(*this);
+	HWEncoderDetails.Encoder->UnregisterListener(*this);
 	Callback = nullptr;
 	return 0;
 }
@@ -111,33 +166,142 @@ int32 FVideoEncoder::Encode(const webrtc::VideoFrame& Frame, const webrtc::Codec
 	// the frame managed to pass encoder queue so disable frame drop notification
 	RawFrame->DisableFrameDropNotification();
 
-	FBufferId BufferId = RawFrame->GetBuffer();
+	AVEncoder::FBufferId BufferId = RawFrame->GetBuffer();
 
-	webrtc::EncodedImage EncodedImage;
-	EncodedImage._completeFrame = true;
-	EncodedImage.SetTimestamp(Frame.timestamp());
-	EncodedImage.ntp_time_ms_ = Frame.ntp_time_ms();
-	EncodedImage.capture_time_ms_ = Frame.render_time_ms();
-	EncodedImage.rotation_ = Frame.rotation();
+	auto EncoderCookie = MakeUnique<FEncoderCookie>();
+	EncoderCookie->EncodedImage._completeFrame = true;
+	EncoderCookie->EncodedImage.SetTimestamp(Frame.timestamp());
+	EncoderCookie->EncodedImage.ntp_time_ms_ = Frame.ntp_time_ms();
+	EncoderCookie->EncodedImage.capture_time_ms_ = Frame.render_time_ms();
+	EncoderCookie->EncodedImage.rotation_ = Frame.rotation();
+	EncoderCookie->EncodedImage.timing_.encode_start_ms = rtc::TimeMicros() / 1000;
 
+	bool bForceKeyFrame = false;
 	if (FrameTypes && (*FrameTypes)[0] == webrtc::kVideoFrameKey)
 	{
 		UE_LOG(PixelStreamer, Verbose, TEXT("key-frame requested, size=%zu"), FrameTypes->size());
-		EncodedImage._frameType = webrtc::kVideoFrameKey;
+		bForceKeyFrame = true;
+	}
+
+	// Adjust framerate if required, so we can keep quality without increasing bitrate
+	{
+		float Fps;
+		if (!CVarPixelStreamingEncoderPrioritiseQuality.GetValueOnAnyThread())
+		{
+			Fps = HWEncoderDetails.InitialMaxFPS;
+		}
+		else
+		{
+
+			// Quality of video suffers if B/W is limited and drops below some threshold. We can sacrifice
+			// responsiveness (latency) to improve video quality. We reduce framerate and so encoder
+			// can spread limited bitrate over fewer frames.
+			int32 Bps = HWEncoderDetails.LastBitrate;
+
+			// bitrate lower than lower bound results always in min FPS
+			// bitrate between lower and upper bounds results in FPS proportionally between min and max FPS
+			// bitrate higher than upper bound results always in max FPS
+			const int32 UpperBoundBps = CVarPixelStreamingEncoderHighBitrate.GetValueOnAnyThread();
+			const int32 LowerBoundBps = FMath::Min(CVarPixelStreamingEncoderLowBitrate.GetValueOnAnyThread(), UpperBoundBps);
+			const float MaxFps = HWEncoderDetails.InitialMaxFPS;
+			const float MinFps = FMath::Min(CVarPixelStreamingEncoderMinFPS.GetValueOnAnyThread(), MaxFps);
+
+			if (Bps < LowerBoundBps)
+				Fps = MinFps;
+			else if (Bps < UpperBoundBps)
+			{
+				Fps = MinFps + (MaxFps - MinFps) / (UpperBoundBps - LowerBoundBps) * (Bps - LowerBoundBps);
+			}
+			else
+			{
+				Fps = MaxFps;
+			}
+		}
+
+		if (HWEncoderDetails.LastFramerate != int(Fps))
+		{
+			HWEncoderDetails.LastFramerate = int(Fps);
+			// SetMaxFPS must be called from the game thread because it changes a console var
+			AsyncTask(ENamedThreads::GameThread, [Fps = int(Fps)]() { GEngine->SetMaxFPS(Fps); });
+			HWEncoderDetails.Encoder->SetFramerate(HWEncoderDetails.LastFramerate);
+		}
+	}
+
+	// Adjust QP
+	{
+		int MinQP = CVarPixelStreamingEncoderMinQP.GetValueOnAnyThread();
+		MinQP = FMath::Clamp(MinQP, 0, 51);
+		if (HWEncoderDetails.LastQP != MinQP)
+		{
+			HWEncoderDetails.LastQP = MinQP;
+			HWEncoderDetails.Encoder->SetParameter(TEXT("qp"), FString::FromInt(HWEncoderDetails.LastQP));
+		}
+	}
+
+	// Adjust RcMode
+	{
+		FString RcMode = CVarPixelStreamingEncoderRateControl.GetValueOnAnyThread();
+		if (HWEncoderDetails.LastRcMode != RcMode)
+		{
+			HWEncoderDetails.LastRcMode = RcMode;
+			HWEncoderDetails.Encoder->SetParameter(TEXT("ratecontrolmode"), HWEncoderDetails.LastRcMode);
+		}
 	}
 
 	// TODO(andriy): `LastBitrate.get_sum_kbps()` most probably includes audio bitrate too,
 	// check if this causes any packet drops
-	HWEncoder.EncodeFrame(BufferId, EncodedImage, LastBitrate.get_sum_bps());
+	HWEncoderDetails.LastBitrate = FMath::Min((uint32)LastBitrate.get_sum_bps(), (uint32)CVarPixelStreamingEncoderMaxBitrate.GetValueOnAnyThread());
+	HWEncoderDetails.Encoder->Encode(BufferId, bForceKeyFrame, HWEncoderDetails.LastBitrate, MoveTemp(EncoderCookie));
 
 	return WEBRTC_VIDEO_CODEC_OK;
 }
 
-void FVideoEncoder::OnEncodedFrame(const webrtc::EncodedImage& EncodedImage)
+void FVideoEncoder::OnEncodedVideoFrame(const AVEncoder::FAVPacket& Packet, AVEncoder::FEncoderVideoFrameCookie* CookieIn)
 {
+	FEncoderCookie* Cookie = static_cast<FEncoderCookie*>(CookieIn);
+
+	// Check if the encoder dropped it for some reason
+	if (!Packet.IsValid())
+	{
+		Callback->OnDroppedFrame(webrtc::EncodedImageCallback::DropReason::kDroppedByEncoder);
+		UE_LOG(PixelStreamer, Log, TEXT("Dropping frame due to encoder failure"));
+		return;
+	}
+
+	// The hardware encoder is shared between webrtc's FVideoEncoder instances, so we only create the EncodedImage buffer once
+	if (Cookie->EncodedFrameBuffer.Num()==0)
+	{
+		Cookie->EncodedImage._encodedWidth = Packet.Video.Width;
+		Cookie->EncodedImage._encodedHeight = Packet.Video.Height;
+		Cookie->EncodedImage.timing_.encode_finish_ms = rtc::TimeMicros() / 1000;
+		Cookie->EncodedImage.timing_.flags = webrtc::VideoSendTiming::kTriggeredByTimer;
+		Cookie->EncodedImage._frameType = Packet.IsVideoKeyFrame() ? webrtc::kVideoFrameKey : webrtc::kVideoFrameDelta;
+		Cookie->EncodedImage.qp_ = Packet.Video.FrameAvgQP;
+		Cookie->EncodedFrameBuffer = Packet.Data;
+		Cookie->EncodedImage._buffer = Cookie->EncodedFrameBuffer.GetData();
+		Cookie->EncodedImage._length = Cookie->EncodedFrameBuffer.Num();
+
+		FHUDStats& Stats = FHUDStats::Get();
+		double LatencyMs = Packet.Timings.EncodeFinishTs.GetTotalMilliseconds() - Packet.Timings.EncodeStartTs.GetTotalMilliseconds();
+		double BitrateMbps = Packet.Data.Num() * 8 * Packet.Video.Framerate / 1000000.0;
+		if (Stats.bEnabled)
+		{
+			Stats.EncoderLatencyMs.Update(LatencyMs);
+			Stats.EncoderBitrateMbps.Update(BitrateMbps);
+			Stats.EncoderQP.Update(Packet.Video.FrameAvgQP);
+		}
+
+		UE_LOG(PixelStreamer, VeryVerbose, TEXT("QP %d/%0.f, latency %.0f/%.0f ms, bitrate %.3f/%.3f Mbps, %d bytes")
+			, Packet.Video.FrameAvgQP, Stats.EncoderQP.Get()
+			, LatencyMs, Stats.EncoderLatencyMs.Get()
+			, BitrateMbps, Stats.EncoderBitrateMbps.Get()
+			, (int)Packet.Data.Num());
+
+	}
+
 	// fill RTP fragmentation info
 	std::vector<webrtc::H264::NaluIndex> NALUIndices =
-		webrtc::H264::FindNaluIndices(EncodedImage._buffer, EncodedImage._length);
+		webrtc::H264::FindNaluIndices(Cookie->EncodedImage._buffer, Cookie->EncodedImage._length);
 	FragHeader.VerifyAndAllocateFragmentationHeader(NALUIndices.size());
 	FragHeader.fragmentationVectorSize = static_cast<uint16_t>(NALUIndices.size());
 	for (int I = 0; I != NALUIndices.size(); ++I)
@@ -147,10 +311,10 @@ void FVideoEncoder::OnEncodedFrame(const webrtc::EncodedImage& EncodedImage)
 		FragHeader.fragmentationLength[I] = NALUIndex.payload_size;
 	}
 
-	UE_LOG(PixelStreamer, VeryVerbose, TEXT("(%d) encoded ts %u, ntp_ts_ms %lld, capture_ts_ms %lld"), RtcTimeMs(), EncodedImage.Timestamp(), EncodedImage.ntp_time_ms_, EncodedImage.capture_time_ms_);
+	UE_LOG(PixelStreamer, VeryVerbose, TEXT("(%d) encoded ts %u, ntp_ts_ms %lld, capture_ts_ms %lld"), RtcTimeMs(), Cookie->EncodedImage.Timestamp(), Cookie->EncodedImage.ntp_time_ms_, Cookie->EncodedImage.capture_time_ms_);
 
 	// Deliver encoded image.
-	Callback->OnEncodedImage(EncodedImage, &CodecSpecific, &FragHeader);
+	Callback->OnEncodedImage(Cookie->EncodedImage, &CodecSpecific, &FragHeader);
 
 	//// lame video recording to a file
 	//static IFileHandle* FileHandle = nullptr;
diff --git a/Engine/Plugins/Media/PixelStreaming/Source/PixelStreaming/Private/Codecs/VideoEncoder.h b/Engine/Plugins/Media/PixelStreaming/Source/PixelStreaming/Private/Codecs/VideoEncoder.h
--- a/Engine/Plugins/Media/PixelStreaming/Source/PixelStreaming/Private/Codecs/VideoEncoder.h
+++ b/Engine/Plugins/Media/PixelStreaming/Source/PixelStreaming/Private/Codecs/VideoEncoder.h
@@ -6,20 +6,32 @@
 
 #include "HAL/ThreadSafeBool.h"
 
-class FPixelStreamingBaseVideoEncoder;
 class FPlayerSession;
 
-class FVideoEncoder : public webrtc::VideoEncoder
+class FVideoEncoder : public webrtc::VideoEncoder, public AVEncoder::IVideoEncoderListener
 {
 public:
-	FVideoEncoder(FPixelStreamingBaseVideoEncoder& HWEncoder, FPlayerSession& OwnerSession);
+
+	struct FEncoderCookie : AVEncoder::FEncoderVideoFrameCookie
+	{
+		virtual ~FEncoderCookie() {}
+		webrtc::EncodedImage EncodedImage;
+		// buffer to hold last encoded frame bitstream, because `webrtc::EncodedImage` doesn't take ownership of
+		// the memory
+		TArray<uint8> EncodedFrameBuffer;
+	};
+
+	FVideoEncoder(FHWEncoderDetails& InHWEncoderDetails, FPlayerSession& OwnerSession);
 	~FVideoEncoder() override;
 
 	bool IsQualityController() const
 	{ return bControlsQuality; }
 	void SetQualityController(bool bControlsQuality);
 
-	void OnEncodedFrame(const webrtc::EncodedImage& EncodedImage);
+	//
+	// AVEncoder::IVideoEncoderListener
+	//
+	void OnEncodedVideoFrame(const AVEncoder::FAVPacket& Packet, AVEncoder::FEncoderVideoFrameCookie* Cookie) override;
 
 	//
 	// webrtc::VideoEncoder interface
@@ -36,7 +48,7 @@ public:
 
 private:
 	// Player session that this encoder instance belongs to
-	FPixelStreamingBaseVideoEncoder& HWEncoder;
+	FHWEncoderDetails& HWEncoderDetails;
 	FPlayerSession* PlayerSession = nullptr;
 	webrtc::EncodedImageCallback* Callback = nullptr;
 	webrtc::CodecSpecificInfo CodecSpecific;
@@ -50,7 +62,7 @@ private:
 class FVideoEncoderFactory : public webrtc::VideoEncoderFactory
 {
 public:
-	explicit FVideoEncoderFactory(FPixelStreamingBaseVideoEncoder& HWEncoder);
+	explicit FVideoEncoderFactory(FHWEncoderDetails& HWEncoderDetails);
 
 	/**
 	* This is used from the FPlayerSession::OnSucess to let the factory know
@@ -67,7 +79,7 @@ public:
 	std::unique_ptr<webrtc::VideoEncoder> CreateVideoEncoder(const webrtc::SdpVideoFormat& Format) override;
 
 private:
-	FPixelStreamingBaseVideoEncoder& HWEncoder;
+	FHWEncoderDetails& HWEncoderDetails;
 	TQueue<FPlayerSession*> PendingPlayerSessions;
 };
 
diff --git a/Engine/Plugins/Media/PixelStreaming/Source/PixelStreaming/Private/Codecs/WmfIncludes.h b/Engine/Plugins/Media/PixelStreaming/Source/PixelStreaming/Private/Codecs/WmfIncludes.h
--- a/Engine/Plugins/Media/PixelStreaming/Source/PixelStreaming/Private/Codecs/WmfIncludes.h
+++ b/Engine/Plugins/Media/PixelStreaming/Source/PixelStreaming/Private/Codecs/WmfIncludes.h
@@ -80,7 +80,6 @@ inline const FString GetComErrorDescription(HRESULT Res)
 		}\
 	}
 
-// #AMF(Andriy) : Rename these to DX , instead of DX9 ?
 #define CHECK_HR_DX9(DX9_call)\
 	{\
 		HRESULT Res = DX9_call;\
diff --git a/Engine/Plugins/Media/PixelStreaming/Source/PixelStreaming/Private/PixelStreamerInputComponent.cpp b/Engine/Plugins/Media/PixelStreaming/Source/PixelStreaming/Private/PixelStreamerInputComponent.cpp
--- a/Engine/Plugins/Media/PixelStreaming/Source/PixelStreaming/Private/PixelStreamerInputComponent.cpp
+++ b/Engine/Plugins/Media/PixelStreaming/Source/PixelStreaming/Private/PixelStreamerInputComponent.cpp
@@ -10,8 +10,6 @@
 #include "Serialization/JsonSerializer.h"
 #include "GameFramework/GameUserSettings.h"
 
-//extern TAutoConsoleVariable<float> CVarStreamerBitrateReduction;
-
 UPixelStreamerInputComponent::UPixelStreamerInputComponent()
 	: PixelStreamingModule(FModuleManager::Get().GetModulePtr<IPixelStreamingModule>("PixelStreaming"))
 {
diff --git a/Engine/Plugins/Media/PixelStreaming/Source/PixelStreaming/Private/PixelStreamingPrivate.h b/Engine/Plugins/Media/PixelStreaming/Source/PixelStreaming/Private/PixelStreamingPrivate.h
--- a/Engine/Plugins/Media/PixelStreaming/Source/PixelStreaming/Private/PixelStreamingPrivate.h
+++ b/Engine/Plugins/Media/PixelStreaming/Source/PixelStreaming/Private/PixelStreamingPrivate.h
@@ -4,7 +4,7 @@
 
 #include "WebRTCIncludes.h"
 #include "Codecs/WmfIncludes.h"
-
+#include "AVEncoder.h"
 #include "HAL/IConsoleManager.h"
 #include "Logging/LogMacros.h"
 
@@ -13,3 +13,19 @@ DECLARE_LOG_CATEGORY_EXTERN(PixelStreamer, Log, All);
 DECLARE_LOG_CATEGORY_EXTERN(PixelPlayer, Log, All);
 
 using FPlayerId = uint32;
+
+//
+// Since a given encoder is shared between sessions, we need a place to put together
+// the encoder and any state that should not be repeated per session
+struct FHWEncoderDetails
+{
+	// Max FPS when Pixel streaming was initialized.
+	// We need to keep this, so we can adjust framerate at runtime to keep quality without increasing bitrate
+	uint32 InitialMaxFPS;
+	uint32 LastBitrate = 0;
+	uint32 LastFramerate = 0;
+	int LastQP = -1;
+	FString LastRcMode = TEXT("");
+	TUniquePtr<AVEncoder::FVideoEncoder> Encoder;
+};
+
diff --git a/Engine/Plugins/Media/PixelStreaming/Source/PixelStreaming/Private/RawFrameBuffer.h b/Engine/Plugins/Media/PixelStreaming/Source/PixelStreaming/Private/RawFrameBuffer.h
--- a/Engine/Plugins/Media/PixelStreaming/Source/PixelStreaming/Private/RawFrameBuffer.h
+++ b/Engine/Plugins/Media/PixelStreaming/Source/PixelStreaming/Private/RawFrameBuffer.h
@@ -2,9 +2,8 @@
 
 #pragma once
 
-#include "Codecs/PixelStreamingBaseVideoEncoder.h"
+#include "AVEncoder.h"
 
-// #AMF : Revise this comment. It mentions NvEnc everywhere
 // WebRTC can drop frames in the encoder queue for various reasons, e.g. when more than one frame is waiting
 // for encoding, or when encoder is not ready yet
 // Our pipeline is asynchronous (cos we use NvEnc async encoding) so to keep track of captured frames
@@ -15,7 +14,7 @@
 // and so NvEnc will be notified immediately.
 struct FFrameDropDetector final
 {
-	FFrameDropDetector(FPixelStreamingBaseVideoEncoder& HWEncoder, FBufferId BufferId):
+	FFrameDropDetector(AVEncoder::FVideoEncoder& HWEncoder, AVEncoder::FBufferId BufferId):
 		HWEncoder(&HWEncoder),
 		BufferId(BufferId)
 	{}
@@ -23,11 +22,11 @@ struct FFrameDropDetector final
 	~FFrameDropDetector()
 	{
 		if (HWEncoder)
-			HWEncoder->OnFrameDropped(BufferId);
+			HWEncoder->Drop(BufferId);
 	}
 
-	FPixelStreamingBaseVideoEncoder* HWEncoder = nullptr;
-	FBufferId BufferId;
+	AVEncoder::FVideoEncoder* HWEncoder = nullptr;
+	AVEncoder::FBufferId BufferId;
 };
 
 class FRawFrameBuffer : public webrtc::VideoFrameBuffer
@@ -67,7 +66,7 @@ public:
 	//
 	// Own methods
 	//
-	FBufferId GetBuffer() const
+	AVEncoder::FBufferId GetBuffer() const
 	{
 		return FrameDropDetector->BufferId;
 	}
diff --git a/Engine/Plugins/Media/PixelStreaming/Source/PixelStreaming/Private/Streamer.cpp b/Engine/Plugins/Media/PixelStreaming/Source/PixelStreaming/Private/Streamer.cpp
--- a/Engine/Plugins/Media/PixelStreaming/Source/PixelStreaming/Private/Streamer.cpp
+++ b/Engine/Plugins/Media/PixelStreaming/Source/PixelStreaming/Private/Streamer.cpp
@@ -1,8 +1,6 @@
 // Copyright 1998-2019 Epic Games, Inc. All Rights Reserved.
 
 #include "Streamer.h"
-#include "Codecs/PixelStreamingNvVideoEncoder.h"
-#include "Codecs/PixelStreamingAmfVideoEncoder.h"
 #include "AudioCapturer.h"
 #include "VideoCapturer.h"
 #include "PlayerSession.h"
@@ -16,12 +14,11 @@
 
 DEFINE_LOG_CATEGORY(PixelStreamer);
 
+extern TAutoConsoleVariable<int32> CVarPixelStreamingEncoderMaxBitrate;
+
 bool FStreamer::CheckPlatformCompatibility()
 {
-	return 
-		FPixelStreamingNvVideoEncoder::CheckPlatformCompatibility() ||
-		FPixelStreamingAmfVideoEncoder::CheckPlatformCompatibility();
-
+	return AVEncoder::FVideoEncoderFactory::FindFactory(TEXT("h264")) ? true : false;
 }
 
 FStreamer::FStreamer(const FString& InSignallingServerUrl):
@@ -29,21 +26,46 @@ FStreamer::FStreamer(const FString& InSignallingServerUrl):
 {
 	RedirectWebRtcLogsToUE4(rtc::LoggingSeverity::LS_VERBOSE);
 
+	FModuleManager::LoadModuleChecked<IModuleInterface>(TEXT("AVEncoder"));
+
 	// required for communication with Signalling Server and must be called in the game thread, while it's used in signalling thread
 	FModuleManager::LoadModuleChecked<FWebSocketsModule>("WebSockets");
 
 	FParse::Value(FCommandLine::Get(), TEXT("WebRtcPlanB"), *reinterpret_cast<uint8*>(&bPlanB));
 
-	if (FPixelStreamingNvVideoEncoder::CheckPlatformCompatibility())
+	AVEncoder::FVideoEncoderFactory* UEVideoEncoderFactory = AVEncoder::FVideoEncoderFactory::FindFactory(TEXT("h264"));
+	if (!UEVideoEncoderFactory)
 	{
-		HWEncoder = MakeUnique<FPixelStreamingNvVideoEncoder>();
+		UE_LOG(PixelStreamer, Fatal, TEXT("No video encoder found"));
 	}
-	else
+
+	HWEncoderDetails.InitialMaxFPS = GEngine->GetMaxFPS();
+	if (HWEncoderDetails.InitialMaxFPS == 0)
+	{
+		check(IsInGameThread());
+		HWEncoderDetails.InitialMaxFPS = 60;
+		GEngine->SetMaxFPS(HWEncoderDetails.InitialMaxFPS);
+	}
+
+	HWEncoderDetails.Encoder = UEVideoEncoderFactory->CreateEncoder(TEXT("h264"));
+	if (!HWEncoderDetails.Encoder)
+	{
+		UE_LOG(PixelStreamer, Fatal, TEXT("Could not create video encoder"));
+	}
+
+	AVEncoder::FVideoEncoderConfig Cfg;
+	Cfg.Width = 1920;
+	Cfg.Height = 1080;
+	Cfg.Framerate = 60;
+	Cfg.MaxBitrate = CVarPixelStreamingEncoderMaxBitrate.GetValueOnAnyThread();
+	Cfg.Bitrate = FMath::Min((uint32)1000000, Cfg.MaxBitrate);
+	Cfg.Preset = AVEncoder::FVideoEncoderConfig::EPreset::LowLatency;
+	if (!HWEncoderDetails.Encoder->Initialize(Cfg))
 	{
-		HWEncoder = MakeUnique<FPixelStreamingAmfVideoEncoder>();
+		UE_LOG(PixelStreamer, Fatal, TEXT("Could not initialize video encoder"));
 	}
 
-	VideoEncoderFactoryStrong = std::make_unique<FVideoEncoderFactory>(*HWEncoder);
+	VideoEncoderFactoryStrong = std::make_unique<FVideoEncoderFactory>(HWEncoderDetails);
 
 	// #HACK: Keep a pointer to the Video encoder factory, so we can use it to figure out the
 	// FPlayerSession <-> FVideoEncoder relationship later on
@@ -58,6 +80,7 @@ FStreamer::~FStreamer()
 	// stop WebRtc WndProc thread
 	PostThreadMessage(WebRtcSignallingThreadId, WM_QUIT, 0, 0);
 	WebRtcSignallingThread->Join();
+	HWEncoderDetails.Encoder->Shutdown();
 }
 
 void FStreamer::WebRtcSignallingThreadFunc()
@@ -274,7 +297,7 @@ void FStreamer::AddStreams(FPlayerId PlayerId)
 
 			Stream->AddTrack(AudioTrackLocal);
 
-			auto VideoCapturerStrong = std::make_unique<FVideoCapturer>(*HWEncoder);
+			auto VideoCapturerStrong = std::make_unique<FVideoCapturer>(HWEncoderDetails);
 			VideoCapturer = VideoCapturerStrong.get();
 			rtc::scoped_refptr<webrtc::VideoTrackInterface> VideoTrackLocal(PeerConnectionFactory->CreateVideoTrack(
 				VideoLabel, PeerConnectionFactory->CreateVideoSource(std::move(VideoCapturerStrong))));
@@ -301,7 +324,7 @@ void FStreamer::AddStreams(FPlayerId PlayerId)
 
 		if (!VideoTrack)
 		{
-			auto VideoCapturerStrong = std::make_unique<FVideoCapturer>(*HWEncoder);
+			auto VideoCapturerStrong = std::make_unique<FVideoCapturer>(HWEncoderDetails);
 			VideoCapturer = VideoCapturerStrong.get();
 			VideoTrack = PeerConnectionFactory->CreateVideoTrack(
 				VideoLabel, PeerConnectionFactory->CreateVideoSource(std::move(VideoCapturerStrong)));
diff --git a/Engine/Plugins/Media/PixelStreaming/Source/PixelStreaming/Private/Streamer.h b/Engine/Plugins/Media/PixelStreaming/Source/PixelStreaming/Private/Streamer.h
--- a/Engine/Plugins/Media/PixelStreaming/Source/PixelStreaming/Private/Streamer.h
+++ b/Engine/Plugins/Media/PixelStreaming/Source/PixelStreaming/Private/Streamer.h
@@ -3,7 +3,6 @@
 #pragma once
 
 #include "PixelStreamingPrivate.h"
-#include "Codecs/PixelStreamingBaseVideoEncoder.h"
 #include "SignallingServerConnection.h"
 #include "ProtocolDefs.h"
 
@@ -73,7 +72,7 @@ private:
 
 	// a single instance of the actual hardware encoder is used for multiple streams/players to provide multicasting functionality
 	// for multi-session unicast implementation each instance of `FVideoEncoder` will have own instance of hardware encoder.
-	TUniquePtr<FPixelStreamingBaseVideoEncoder> HWEncoder;
+	FHWEncoderDetails HWEncoderDetails;
 
 	TUniquePtr<FThread> WebRtcSignallingThread;
 	DWORD WebRtcSignallingThreadId = 0;
diff --git a/Engine/Plugins/Media/PixelStreaming/Source/PixelStreaming/Private/VideoCapturer.cpp b/Engine/Plugins/Media/PixelStreaming/Source/PixelStreaming/Private/VideoCapturer.cpp
--- a/Engine/Plugins/Media/PixelStreaming/Source/PixelStreaming/Private/VideoCapturer.cpp
+++ b/Engine/Plugins/Media/PixelStreaming/Source/PixelStreaming/Private/VideoCapturer.cpp
@@ -2,19 +2,23 @@
 
 #include "VideoCapturer.h"
 #include "RawFrameBuffer.h"
-#include "Codecs/PixelStreamingBaseVideoEncoder.h"
 #include "Utils.h"
 
 #include "HAL/PlatformTime.h"
 #include "Misc/Timespan.h"
 
-FVideoCapturer::FVideoCapturer(FPixelStreamingBaseVideoEncoder& InHWEncoder): HWEncoder(InHWEncoder)
+extern TAutoConsoleVariable<int32> CVarPixelStreamingEncoderUseBackBufferSize;
+extern TAutoConsoleVariable<FString> CVarPixelStreamingEncoderTargetSize;
+
+FVideoCapturer::FVideoCapturer(FHWEncoderDetails& InHWEncoderDetails): HWEncoderDetails(InHWEncoderDetails)
 {
 	set_enable_video_adapter(false);
 
 	std::vector<cricket::VideoFormat> Formats;
 	Formats.push_back(cricket::VideoFormat(Width, Height, cricket::VideoFormat::FpsToInterval(Framerate), cricket::FOURCC_H264));
 	SetSupportedFormats(Formats);
+
+	LastTimestampUs = rtc::TimeMicros();
 }
 
 cricket::CaptureState FVideoCapturer::Start(const cricket::VideoFormat& Format)
@@ -26,13 +30,38 @@ void FVideoCapturer::OnFrameReady(const FTexture2DRHIRef& FrameBuffer)
 {
 	int64 TimestampUs = rtc::TimeMicros();
 
-	FBufferId BufferId;
-	if (!HWEncoder.CopyBackBuffer(FrameBuffer, FTimespan::FromMicroseconds(TimestampUs), BufferId))
+	FIntPoint Resolution = FrameBuffer->GetSizeXY();
+	if (CVarPixelStreamingEncoderUseBackBufferSize.GetValueOnRenderThread() == 0)
+	{
+		FString EncoderTargetSize = CVarPixelStreamingEncoderTargetSize.GetValueOnRenderThread();
+		FString TargetWidth, TargetHeight;
+		bool bValidSize = EncoderTargetSize.Split(TEXT("x"), &TargetWidth, &TargetHeight);
+		if (bValidSize)
+		{
+			Resolution.X = FCString::Atoi(*TargetWidth);
+			Resolution.Y = FCString::Atoi(*TargetHeight);
+		}
+		else
+		{
+			UE_LOG(PixelStreamer, Error, TEXT("CVarPixelStreamingEncoderTargetSize is not in a valid format: %s. It should be e.g: \"1920x1080\""), *EncoderTargetSize);
+			CVarPixelStreamingEncoderTargetSize->Set(*FString::Printf(TEXT("%dx%d"), Resolution.X, Resolution.Y));
+		}
+	}
+
+	AVEncoder::FBufferId BufferId;
+	if (!HWEncoderDetails.Encoder->CopyTexture(
+		FrameBuffer,
+		FTimespan::FromMicroseconds(TimestampUs),
+		FTimespan::FromMicroseconds(TimestampUs - LastTimestampUs),
+		BufferId,
+		Resolution))
 	{
 		return;
 	}
 
-	rtc::scoped_refptr<FRawFrameBuffer> Buffer = new rtc::RefCountedObject<FRawFrameBuffer>(MakeShared<FFrameDropDetector>(HWEncoder, BufferId), FrameBuffer->GetSizeX(), FrameBuffer->GetSizeY());
+	LastTimestampUs = TimestampUs;
+
+	rtc::scoped_refptr<FRawFrameBuffer> Buffer = new rtc::RefCountedObject<FRawFrameBuffer>(MakeShared<FFrameDropDetector>(*HWEncoderDetails.Encoder, BufferId), FrameBuffer->GetSizeX(), FrameBuffer->GetSizeY());
 
 	webrtc::VideoFrame Frame = webrtc::VideoFrame::Builder().
 		set_video_frame_buffer(Buffer).
@@ -42,3 +71,4 @@ void FVideoCapturer::OnFrameReady(const FTexture2DRHIRef& FrameBuffer)
 	UE_LOG(PixelStreamer, VeryVerbose, TEXT("(%d) captured video %lld"), RtcTimeMs(), TimestampUs);
 	OnFrame(Frame, Buffer->width(), Buffer->height());
 }
+
diff --git a/Engine/Plugins/Media/PixelStreaming/Source/PixelStreaming/Private/VideoCapturer.h b/Engine/Plugins/Media/PixelStreaming/Source/PixelStreaming/Private/VideoCapturer.h
--- a/Engine/Plugins/Media/PixelStreaming/Source/PixelStreaming/Private/VideoCapturer.h
+++ b/Engine/Plugins/Media/PixelStreaming/Source/PixelStreaming/Private/VideoCapturer.h
@@ -5,12 +5,10 @@
 #include "PixelStreamingPrivate.h"
 #include "RHIResources.h"
 
-class FPixelStreamingBaseVideoEncoder;
-
 class FVideoCapturer : public cricket::VideoCapturer
 {
 public:
-	explicit FVideoCapturer(FPixelStreamingBaseVideoEncoder& InHWEncoder);
+	explicit FVideoCapturer(FHWEncoderDetails& InHWEncoderDetails);
 
 	void OnFrameReady(const FTexture2DRHIRef& FrameBuffer);
 
@@ -35,7 +33,8 @@ private:
 	}
 	//////////////////////////////////////////////////////////////////////////
 
-	FPixelStreamingBaseVideoEncoder& HWEncoder;
+	FHWEncoderDetails& HWEncoderDetails;
+	int64 LastTimestampUs = 0;
 
 	int32 Width = 1920;
 	int32 Height = 1080;
diff --git a/Engine/Source/Runtime/AVEncoder/AVEncoder.Build.cs b/Engine/Source/Runtime/AVEncoder/AVEncoder.Build.cs
new file mode 100644
--- /dev/null
+++ b/Engine/Source/Runtime/AVEncoder/AVEncoder.Build.cs
@@ -0,0 +1,36 @@
+// Copyright 1998-2019 Epic Games, Inc. All Rights Reserved.
+
+using System.IO;
+using UnrealBuildTool;
+
+public class AVEncoder : ModuleRules
+{
+	public AVEncoder(ReadOnlyTargetRules Target) : base(Target)
+	{
+
+		PrivateDependencyModuleNames.AddRange(new string[]
+		{
+			"Core",
+			"Engine",
+			"RHI",
+			"RenderCore",
+		});
+
+		if (Target.Platform == UnrealTargetPlatform.Win32 || Target.Platform == UnrealTargetPlatform.Win64)
+		{
+			PrivateDependencyModuleNames.AddRange(new string[]
+				{
+					"D3D11RHI"
+				});
+
+			PublicDelayLoadDLLs.Add("mfplat.dll");
+			PublicDelayLoadDLLs.Add("mfuuid.dll");
+			PublicDelayLoadDLLs.Add("Mfreadwrite.dll");
+
+			PublicSystemLibraries.Add("d3d11.lib");
+			PublicSystemLibraries.Add("DXGI.lib");
+		}
+
+	}
+}
+
diff --git a/Engine/Source/Runtime/AVEncoder/Private/AVEncoderCommon.cpp b/Engine/Source/Runtime/AVEncoder/Private/AVEncoderCommon.cpp
new file mode 100644
--- /dev/null
+++ b/Engine/Source/Runtime/AVEncoder/Private/AVEncoderCommon.cpp
@@ -0,0 +1,125 @@
+// Copyright 1998-2019 Epic Games, Inc. All Rights Reserved.
+
+#include "AVEncoderCommon.h"
+
+DEFINE_LOG_CATEGORY(LogAVEncoder);
+CSV_DECLARE_CATEGORY_EXTERN(AVEncoder);
+CSV_DEFINE_CATEGORY(AVEncoder, true);
+
+namespace AVEncoder
+{
+
+bool ReadH264Setting(const FString& Name, const FString& Value, FH264Settings& OutSettings)
+{
+	if (Name == TEXT("qp"))
+	{
+		OutSettings.QP =  FMath::Clamp(FCString::Atoi(*Value), 0, 51);
+	}
+	else if (Name == TEXT("ratecontrolmode"))
+	{
+		if (Value == TEXT("constqp"))
+		{
+			OutSettings.RcMode = FH264Settings::ConstQP;
+		}
+		else if (Value == TEXT("vbr"))
+		{
+			OutSettings.RcMode = FH264Settings::VBR;
+		}
+		else if (Value == TEXT("cbr"))
+		{
+			OutSettings.RcMode = FH264Settings::CBR;
+		}
+		else
+		{
+			UE_LOG(LogAVEncoder, Error, TEXT("Option '%s' has an invalid value ('%s')"), *Name, *Value);
+			return false;
+		}
+	}
+	else
+	{
+		return false;
+	}
+
+	return true;
+}
+
+void ReadH264Settings(const TArray<TPair<FString, FString>>& Options, FH264Settings& OutSettings)
+{
+	for (const TPair<FString, FString>& Opt : Options)
+	{
+		ReadH264Setting(Opt.Key, Opt.Value, OutSettings);
+	}
+}
+
+void CopyTextureImpl(const FTexture2DRHIRef& Src, FTexture2DRHIRef& Dst, FRHIGPUFence* GPUFence)
+{
+	IRendererModule* RendererModule = &FModuleManager::GetModuleChecked<IRendererModule>("Renderer");
+	FRHICommandListImmediate& RHICmdList = FRHICommandListExecutor::GetImmediateCommandList();
+
+	if (GPUFence)
+	{
+		GPUFence->Clear();
+	}
+
+	if (Src->GetFormat() == Dst->GetFormat() &&
+		Src->GetSizeXY() == Dst->GetSizeXY())
+	{
+		RHICmdList.CopyToResolveTarget(Src, Dst, FResolveParams{});
+	}
+	else // Texture format mismatch, use a shader to do the copy.
+	{
+		// #todo-renderpasses there's no explicit resolve here? Do we need one?
+		FRHIRenderPassInfo RPInfo(Dst, ERenderTargetActions::Load_Store);
+		RHICmdList.BeginRenderPass(RPInfo, TEXT("CopyBackbuffer"));
+		{
+			RHICmdList.SetViewport(0, 0, 0.0f, Dst->GetSizeX(), Dst->GetSizeY(), 1.0f);
+
+			FGraphicsPipelineStateInitializer GraphicsPSOInit;
+			RHICmdList.ApplyCachedRenderTargets(GraphicsPSOInit);
+			GraphicsPSOInit.BlendState = TStaticBlendState<>::GetRHI();
+			GraphicsPSOInit.RasterizerState = TStaticRasterizerState<>::GetRHI();
+			GraphicsPSOInit.DepthStencilState = TStaticDepthStencilState<false, CF_Always>::GetRHI();
+
+			TShaderMap<FGlobalShaderType>* ShaderMap = GetGlobalShaderMap(GMaxRHIFeatureLevel);
+			TShaderMapRef<FScreenVS> VertexShader(ShaderMap);
+			TShaderMapRef<FScreenPS> PixelShader(ShaderMap);
+
+			GraphicsPSOInit.BoundShaderState.VertexDeclarationRHI = GFilterVertexDeclaration.VertexDeclarationRHI;
+			GraphicsPSOInit.BoundShaderState.VertexShaderRHI = GETSAFERHISHADER_VERTEX(*VertexShader);
+			GraphicsPSOInit.BoundShaderState.PixelShaderRHI = GETSAFERHISHADER_PIXEL(*PixelShader);
+			GraphicsPSOInit.PrimitiveType = PT_TriangleList;
+
+			SetGraphicsPipelineState(RHICmdList, GraphicsPSOInit);
+
+			if (Dst->GetSizeX() != Src->GetSizeX() || Dst->GetSizeY() != Src->GetSizeY())
+			{
+				PixelShader->SetParameters(RHICmdList, TStaticSamplerState<SF_Bilinear>::GetRHI(), Src);
+			}
+			else
+			{
+				PixelShader->SetParameters(RHICmdList, TStaticSamplerState<SF_Point>::GetRHI(), Src);
+			}
+
+			RendererModule->DrawRectangle(
+				RHICmdList,
+				0, 0,						// Dest X, Y
+				Dst->GetSizeX(),			// Dest Width
+				Dst->GetSizeY(),			// Dest Height
+				0, 0,						// Source U, V
+				1, 1,						// Source USize, VSize
+				Dst->GetSizeXY(),			// Target buffer size
+				FIntPoint(1, 1),			// Source texture size
+				*VertexShader,
+				EDRF_Default);
+		}
+		RHICmdList.EndRenderPass();
+
+		if (GPUFence)
+		{
+			RHICmdList.WriteGPUFence(GPUFence);
+		}
+	}
+}
+
+
+}
diff --git a/Engine/Source/Runtime/AVEncoder/Private/AVEncoderCommon.h b/Engine/Source/Runtime/AVEncoder/Private/AVEncoderCommon.h
new file mode 100644
--- /dev/null
+++ b/Engine/Source/Runtime/AVEncoder/Private/AVEncoderCommon.h
@@ -0,0 +1,76 @@
+// Copyright 1998-2019 Epic Games, Inc. All Rights Reserved.
+
+#pragma once
+
+#include "CoreMinimal.h"
+#include "Stats/Stats.h"
+#include "Logging/LogMacros.h"
+#include "RHIStaticStates.h"
+#include "HAL/ThreadSafeBool.h"
+#include "HAL/ThreadSafeCounter.h"
+#include "HAL/Thread.h"
+#include "HAL/Event.h"
+#include "Misc/ScopeExit.h"
+#include "ShaderCore.h"
+#include "ProfilingDebugging/CsvProfiler.h"
+#include "RHI.h"
+#include "RHIResources.h"
+#include "ScreenRendering.h"
+#include "CommonRenderResources.h"
+
+DECLARE_LOG_CATEGORY_EXTERN(LogAVEncoder, Log, All);
+CSV_DECLARE_CATEGORY_EXTERN(AVEncoder);
+
+namespace AVEncoder
+{
+
+template<typename F>
+inline void ExecuteRHICommand(F&& Functor)
+{
+	FRHICommandList& RHICmdList = FRHICommandListExecutor::GetImmediateCommandList();
+	if (RHICmdList.Bypass())
+	{
+		Functor();
+	}
+	else
+	{
+		FRHICOMMAND_MACRO(FLocalRHICommand)
+		{
+			F Functor;
+
+			FLocalRHICommand(F InFunctor)
+				: Functor(MoveTemp(InFunctor))
+			{}
+
+			void Execute(FRHICommandListBase& CmdList)
+			{
+				Functor();
+			}
+		};
+
+		new (RHICmdList.AllocCommand<FLocalRHICommand>()) FLocalRHICommand(Functor);
+	}
+}
+
+// Settings specific to  
+struct FH264Settings
+{
+	enum ERateControlMode
+	{
+		ConstQP,
+		VBR,
+		CBR
+	};
+
+	int QP = 20;
+	ERateControlMode RcMode = ERateControlMode::CBR;
+};
+
+bool ReadH264Setting(const FString& Name, const FString& Value, FH264Settings& OutSettings);
+void ReadH264Settings(const TArray<TPair<FString, FString>>& Options, FH264Settings& OutSettings);
+
+void CopyTextureImpl(const FTexture2DRHIRef& Src, FTexture2DRHIRef& Dst, FRHIGPUFence* GPUFence);
+
+}
+
+
diff --git a/Engine/Source/Runtime/AVEncoder/Private/AVEncoderModule.cpp b/Engine/Source/Runtime/AVEncoder/Private/AVEncoderModule.cpp
new file mode 100644
--- /dev/null
+++ b/Engine/Source/Runtime/AVEncoder/Private/AVEncoderModule.cpp
@@ -0,0 +1,15 @@
+// Copyright 1998-2019 Epic Games, Inc. All Rights Reserved.
+
+#include "Modules/ModuleManager.h"
+#include "AVEncoder.h"
+#include "AVEncoderCommon.h"
+
+class FAVEncoderModule : public IModuleInterface
+{
+public:
+	// Factories initialization is done lazily whenever some code uses the Factory API.
+	// This is so because factories needs some RHI stuff initialized
+};
+
+IMPLEMENT_MODULE(FAVEncoderModule, AVEncoder);
+
diff --git a/Engine/Source/Runtime/AVEncoder/Private/Microsoft/AVEncoderIMFSampleWrapper.cpp b/Engine/Source/Runtime/AVEncoder/Private/Microsoft/AVEncoderIMFSampleWrapper.cpp
new file mode 100644
--- /dev/null
+++ b/Engine/Source/Runtime/AVEncoder/Private/Microsoft/AVEncoderIMFSampleWrapper.cpp
@@ -0,0 +1,122 @@
+// Copyright 1998-2019 Epic Games, Inc. All Rights Reserved.
+
+#include "Microsoft/AVEncoderIMFSampleWrapper.h"
+#include "Microsoft/AVEncoderMicrosoftCommon.h"
+
+namespace AVEncoder
+{
+
+namespace
+{
+	TRefCountPtr<IMFMediaBuffer> CloneMediaBuffer(const TRefCountPtr<IMFMediaBuffer>& Src)
+	{
+		uint8* SrcByteBuffer = nullptr;
+		DWORD SrcMaxLen = 0;
+		DWORD SrcCurrLen = 0;
+		verify(SUCCEEDED(Src->Lock(&SrcByteBuffer, &SrcMaxLen, &SrcCurrLen)));
+
+		TRefCountPtr<IMFMediaBuffer> Dest;
+		verify(SUCCEEDED(MFCreateMemoryBuffer(SrcCurrLen, Dest.GetInitReference())));
+
+		uint8* DestByteBuffer = nullptr;
+		verify(SUCCEEDED(Dest->Lock(&DestByteBuffer, nullptr, nullptr)));
+		FMemory::Memcpy(DestByteBuffer, SrcByteBuffer, SrcCurrLen);
+		verify(SUCCEEDED(Dest->Unlock()));
+		verify(SUCCEEDED(Src->Unlock()));
+
+		verify(SUCCEEDED(Dest->SetCurrentLength(SrcCurrLen)));
+		return Dest;
+	}
+}
+
+//////////////////////////////////////////////////////////////////////////
+// FIMFSampleWrapper
+//////////////////////////////////////////////////////////////////////////
+
+
+bool FIMFSampleWrapper::CreateSample()
+{
+	Sample = nullptr;
+	CHECK_HR(MFCreateSample(Sample.GetInitReference()));
+	return true;
+}
+
+FTimespan FIMFSampleWrapper::GetTime() const
+{
+	int64 Time;
+	HRESULT hr = Sample->GetSampleTime(&Time);
+	verifyf(SUCCEEDED(hr), TEXT("%s"), *GetComErrorDescription(hr));
+	return Time;
+}
+
+void FIMFSampleWrapper::SetTime(FTimespan Time)
+{
+	verify(SUCCEEDED(Sample->SetSampleTime(Time.GetTicks())));
+}
+
+FTimespan FIMFSampleWrapper::GetDuration() const
+{
+	int64 Duration;
+	HRESULT hr = Sample->GetSampleDuration(&Duration);
+	verifyf(SUCCEEDED(hr), TEXT("%s"), *GetComErrorDescription(hr));
+	return Duration;
+}
+
+void FIMFSampleWrapper::SetDuration(FTimespan Duration)
+{
+	verify(SUCCEEDED(Sample->SetSampleDuration(Duration.GetTicks())));
+}
+
+bool FIMFSampleWrapper::IsVideoKeyFrame() const
+{
+	return MediaType == EPacketType::Video && MFGetAttributeUINT32(Sample, MFSampleExtension_CleanPoint, 0) != 0;
+}
+
+int FIMFSampleWrapper::GetBufferCount() const
+{
+	DWORD BufferCount = 0;
+	verify(SUCCEEDED(Sample->GetBufferCount(&BufferCount)));
+	return BufferCount;
+}
+
+FIMFSampleWrapper FIMFSampleWrapper::Clone() const
+{
+	FIMFSampleWrapper Res{ MediaType };
+	if (!Sample)
+	{
+		return Res;
+	}
+
+	verify(Res.CreateSample());
+
+	DWORD Flags = 0;
+	verify(SUCCEEDED(Sample->GetSampleFlags(&Flags)));
+	verify(SUCCEEDED(Res.GetSample()->SetSampleFlags(Flags)));
+
+	Res.SetTime(GetTime());
+	Res.SetDuration(GetDuration());
+
+	if (IsVideoKeyFrame())
+	{
+		verify(SUCCEEDED(Res.GetSample()->SetUINT32(MFSampleExtension_CleanPoint, 1)));
+	}
+
+	DWORD BufferCount = GetBufferCount();
+
+	for (DWORD i = 0; i != BufferCount; ++i)
+	{
+		TRefCountPtr<IMFMediaBuffer> SrcBuf;
+		verify(SUCCEEDED(Sample->GetBufferByIndex(i, SrcBuf.GetInitReference())));
+
+		TRefCountPtr<IMFMediaBuffer> DestBuf = CloneMediaBuffer(SrcBuf);
+		verify(SUCCEEDED(Res.GetSample()->AddBuffer(DestBuf)));
+	}
+
+	return Res;
+}
+
+} // namespace AVEncoder
+
+
+
+
diff --git a/Engine/Source/Runtime/AVEncoder/Private/Microsoft/AVEncoderMicrosoftCommon.cpp b/Engine/Source/Runtime/AVEncoder/Private/Microsoft/AVEncoderMicrosoftCommon.cpp
new file mode 100644
--- /dev/null
+++ b/Engine/Source/Runtime/AVEncoder/Private/Microsoft/AVEncoderMicrosoftCommon.cpp
@@ -0,0 +1,43 @@
+// Copyright 1998-2019 Epic Games, Inc. All Rights Reserved.
+
+#pragma once
+
+#include "AVEncoderMicrosoftCommon.h"
+
+#if WMFMEDIA_SUPPORTED_PLATFORM
+	#pragma comment(lib, "mfplat")
+	#pragma comment(lib, "mfuuid")
+	#pragma comment(lib, "Mfreadwrite")
+#endif
+
+namespace AVEncoder
+{
+
+//
+// Windows only code
+//
+#if PLATFORM_WINDOWS
+ID3D11Device* GetUE4DxDevice()
+{
+	auto Device = static_cast<ID3D11Device*>(GDynamicRHI->RHIGetNativeDevice());
+	checkf(Device != nullptr, TEXT("Failed to get UE4's ID3D11Device"));
+	return Device;
+}
+#endif
+
+//
+// XboxOne only code
+// 
+#if PLATFORM_XBOXONE 
+ID3D12Device* GetUE4DxDevice()
+{
+	auto Device = static_cast<ID3D12Device*>(GDynamicRHI->RHIGetNativeDevice());
+	checkf(Device != nullptr, TEXT("Failed to get UE4's ID3D12Device"));
+	return Device;
+}
+
+#endif
+
+
+} // namespace AVEncoder
+
diff --git a/Engine/Source/Runtime/AVEncoder/Private/Microsoft/AVEncoderMicrosoftCommon.h b/Engine/Source/Runtime/AVEncoder/Private/Microsoft/AVEncoderMicrosoftCommon.h
new file mode 100644
--- /dev/null
+++ b/Engine/Source/Runtime/AVEncoder/Private/Microsoft/AVEncoderMicrosoftCommon.h
@@ -0,0 +1,212 @@
+// Copyright 1998-2019 Epic Games, Inc. All Rights Reserved.
+
+#pragma once
+
+#include "AVEncoder.h"
+#include "AVEncoderCommon.h"
+
+//
+// Windows only include
+//
+#if PLATFORM_WINDOWS
+	THIRD_PARTY_INCLUDES_START
+	#include "Windows/AllowWindowsPlatformTypes.h"
+	#include "Windows/PreWindowsApi.h"
+		#include <d3d11.h>
+		#include <mftransform.h>
+		#include <mfapi.h>
+		#include <mferror.h>
+		#include <mfidl.h>
+		#include <codecapi.h>
+		#include <shlwapi.h>
+		#include <mfreadwrite.h>
+		#include <d3d11_1.h>
+		#include <d3d12.h>
+		#include <dxgi1_4.h>
+	#include "Windows/PostWindowsApi.h"
+	#include "Windows/HideWindowsPlatformTypes.h"
+	THIRD_PARTY_INCLUDES_END
+
+	#include "D3D11State.h"
+	#include "D3D11Resources.h"
+#endif
+
+//
+// XboxOne only includes
+//
+#if PLATFORM_XBOXONE
+
+#include "XboxOne/XboxOneAllowPlatformTypes.h"
+#include "XboxOne/XboxOnePreApi.h"
+		#include <d3d11_x.h>
+		#include <d3d12_x.h>
+		#include <d3dcompiler_x.h>
+		#include <d3dx12_x.h>
+		#include <mftransform.h>
+		#include <mfapi.h>
+		#include <mferror.h>
+		#include <mfidl.h>
+		#include <codecapi.h>
+		#include <mfreadwrite.h>
+#include "XboxOne/XboxOnePostApi.h"
+#include "XboxOne/XboxOneHidePlatformTypes.h"
+
+#endif
+
+#define WMFMEDIA_SUPPORTED_PLATFORM (PLATFORM_WINDOWS && (WINVER >= 0x0600 /*Vista*/) && !UE_SERVER)
+
+namespace AVEncoder
+{
+
+inline const FString GetComErrorDescription(HRESULT Res)
+{
+	const uint32 BufSize = 4096;
+	WIDECHAR buffer[4096];
+	if (::FormatMessageW(FORMAT_MESSAGE_FROM_SYSTEM,
+		nullptr,
+		Res,
+		MAKELANGID(LANG_NEUTRAL, SUBLANG_NEUTRAL),
+		buffer,
+		sizeof(buffer) / sizeof(*buffer),
+		nullptr))
+	{
+		return buffer;
+	}
+	else
+	{
+		return TEXT("[cannot find error description]");
+	}
+}
+
+#if PLATFORM_WINDOWS
+	#include "Windows/AllowWindowsPlatformTypes.h"
+#elif PLATFORM_XBOXONE
+	#include "XboxOne/XboxOneAllowPlatformTypes.h"
+#endif
+
+// macro to deal with COM calls inside a function that returns `false` on error
+#define CHECK_HR(COM_call)\
+	{\
+		HRESULT Res = COM_call;\
+		if (FAILED(Res))\
+		{\
+			UE_LOG(LogAVEncoder, Error, TEXT("`" #COM_call "` failed: 0x%X - %s"), Res, *GetComErrorDescription(Res));\
+			return false;\
+		}\
+	}
+
+// macro to deal with COM calls inside a function that returns `{}` on error
+#define CHECK_HR_DEFAULT(COM_call)\
+	{\
+		HRESULT Res = COM_call;\
+		if (FAILED(Res))\
+		{\
+			UE_LOG(LogAVEncoder, Error, TEXT("`" #COM_call "` failed: 0x%X - %s"), Res, *GetComErrorDescription(Res));\
+			return {};\
+		}\
+	}
+
+// macro to deal with COM calls inside COM method (that returns HRESULT)
+#define CHECK_HR_COM(COM_call)\
+	{\
+		HRESULT Res = COM_call;\
+		if (FAILED(Res))\
+		{\
+			UE_LOG(LogAVEncoder, Error, TEXT("`" #COM_call "` failed: 0x%X - %s"), Res, *GetComErrorDescription(Res));\
+			return Res;\
+		}\
+	}
+
+// macro to deal with COM calls inside COM method (that simply returns)
+#define CHECK_HR_VOID(COM_call)\
+	{\
+		HRESULT Res = COM_call;\
+		if (FAILED(Res))\
+		{\
+			UE_LOG(LogAVEncoder, Error, TEXT("`" #COM_call "` failed: 0x%X - %s"), Res, *GetComErrorDescription(Res));\
+			return;\
+		}\
+	}
+
+#if PLATFORM_WINDOWS
+	#include "Windows/HideWindowsPlatformTypes.h"
+#elif PLATFORM_XBOXONE
+	#include "XboxOne/XboxOneHidePlatformTypes.h"
+#endif
+
+// following commented include causes name clash between UE4 and Windows `IMediaEventSink`,
+// we just need a couple of GUIDs from there so the solution is to duplicate them below
+//#include "wmcodecdsp.h"
+
+const GUID CLSID_AACMFTEncoder = { 0x93AF0C51, 0x2275, 0x45d2, { 0xA3, 0x5B, 0xF2, 0xBA, 0x21, 0xCA, 0xED, 0x00 } };
+const GUID CLSID_MP3ACMCodecWrapper ={ 0x11103421, 0x354c, 0x4cca, { 0xa7, 0xa3, 0x1a, 0xff, 0x9a, 0x5b, 0x67, 0x01 } };
+const GUID CLSID_CMSH264EncoderMFT = { 0x6ca50344, 0x051a, 0x4ded, { 0x97, 0x79, 0xa4, 0x33, 0x05, 0x16, 0x5e, 0x35 } };
+const GUID CLSID_VideoProcessorMFT = { 0x88753b26, 0x5b24, 0x49bd, { 0xb2, 0xe7, 0xc, 0x44, 0x5c, 0x78, 0xc9, 0x82 } };
+
+// `MF_LOW_LATENCY` is defined in "mfapi.h" for >= WIN8
+// UE4 supports lower Windows versions at the moment and so `WINVER` is < `_WIN32_WINNT_WIN8`
+// to be able to use `MF_LOW_LATENCY` with default UE4 build we define it ourselves and check actual
+// Windows version in runtime
+#if (WINVER < _WIN32_WINNT_WIN8)
+	const GUID MF_LOW_LATENCY = { 0x9c27891a, 0xed7a, 0x40e1,{ 0x88, 0xe8, 0xb2, 0x27, 0x27, 0xa0, 0x24, 0xee } };
+#endif
+
+
+#if PLATFORM_WINDOWS
+
+ID3D11Device* GetUE4DxDevice();
+
+#elif PLATFORM_XBOXONE
+
+ID3D12Device* GetUE4DxDevice();
+
+#endif
+
+// scope-disable particular DX11 Debug Layer errors
+#if PLATFORM_WINDOWS
+class FScopeDisabledDxDebugErrors final
+{
+private:
+
+public:
+	FScopeDisabledDxDebugErrors(TArray<D3D11_MESSAGE_ID>&& ErrorsToDisable)
+	{
+		TRefCountPtr<ID3D11Debug> Debug;
+		HRESULT HRes = GetUE4DxDevice()->QueryInterface(__uuidof(ID3D11Debug), reinterpret_cast<void**>(Debug.GetInitReference()));
+
+		if (HRes == E_NOINTERFACE)
+		{
+			// Debug Layer is not enabled, so no need to disable its errors
+			return;
+		}
+
+		if (!SUCCEEDED(HRes) ||
+			!SUCCEEDED(HRes = Debug->QueryInterface(__uuidof(ID3D11InfoQueue), reinterpret_cast<void**>(InfoQueue.GetInitReference()))))
+		{
+			UE_LOG(LogAVEncoder, VeryVerbose, TEXT("Failed to get ID3D11InfoQueue: 0x%X - %s"), HRes, *GetComErrorDescription(HRes));
+			return;
+		}
+
+		D3D11_INFO_QUEUE_FILTER filter = {};
+		filter.DenyList.NumIDs = ErrorsToDisable.Num();
+		filter.DenyList.pIDList = ErrorsToDisable.GetData();
+		bSucceeded = SUCCEEDED(InfoQueue->PushStorageFilter(&filter));
+	}
+
+	~FScopeDisabledDxDebugErrors()
+	{
+		if (bSucceeded)
+		{
+			InfoQueue->PopStorageFilter();
+		}
+	}
+
+private:
+	TRefCountPtr<ID3D11InfoQueue> InfoQueue;
+	bool bSucceeded = false;
+};
+#endif
+
+
+} // namespace AVEncoder
+
diff --git a/Engine/Source/Runtime/AVEncoder/Private/Microsoft/Windows/AmfVideoEncoder.cpp b/Engine/Source/Runtime/AVEncoder/Private/Microsoft/Windows/AmfVideoEncoder.cpp
new file mode 100644
--- /dev/null
+++ b/Engine/Source/Runtime/AVEncoder/Private/Microsoft/Windows/AmfVideoEncoder.cpp
@@ -0,0 +1,867 @@
+// Copyright 1998-2019 Epic Games, Inc. All Rights Reserved.
+
+#include "AmfVideoEncoder.h"
+#include "Microsoft/AVEncoderMicrosoftCommon.h"
+#include "CommonRenderResources.h"
+
+THIRD_PARTY_INCLUDES_START
+	#include "AmdAmf/core/Result.h"
+	#include "AmdAmf/core/Factory.h"
+	#include "AmdAmf/components/VideoEncoderVCE.h"
+	#include "AmdAmf/core/Compute.h"
+	#include "AmdAmf/core/Plane.h"
+THIRD_PARTY_INCLUDES_END
+
+DECLARE_STATS_GROUP(TEXT("AmfVideoEncoder"), STATGROUP_AmfVideoEncoder, STATCAT_Advanced);
+DECLARE_CYCLE_STAT(TEXT("AmfEncoder->QueryOutput"), STAT_Amf_QueryOutput, STATGROUP_AmfVideoEncoder);
+DECLARE_CYCLE_STAT(TEXT("StreamEncodedFrame"), STAT_Amf_StreamEncodedFrame, STATGROUP_AmfVideoEncoder);
+DECLARE_CYCLE_STAT(TEXT("OnEncodedVideoFrameCallback"), STAT_Amf_OnEncodedVideoFrameCallback, STATGROUP_AmfVideoEncoder);
+DECLARE_CYCLE_STAT(TEXT("SubmitFrameToEncoder"), STAT_Amf_SubmitFrameToEncoder, STATGROUP_AmfVideoEncoder);
+DECLARE_CYCLE_STAT(TEXT("AmfEncoder->SubmitInput"), STAT_Amf_SubmitInput, STATGROUP_AmfVideoEncoder);
+
+// NOTE: This only exists in a more recent version of the AMF SDK.  Adding it here so I don't need to update the SDK  yet.
+#ifdef AMF_VIDEO_ENCODER_LOWLATENCY_MODE
+	// If you get this error, it means the AMF SDK was updated and already contains the property, so you can remove this check and the property that was added explicitly.
+	#error "AMF_VIDEO_ENCODER_LOWLATENCY_MODE already exists. Please remove duplicate"
+#endif
+#define AMF_VIDEO_ENCODER_LOWLATENCY_MODE                       L"LowLatencyInternal"       // bool; default = false, enables low latency mode and POC mode 2 in the encoder
+
+#define CHECK_AMF_RET(AMF_call)\
+{\
+	AMF_RESULT Res = AMF_call;\
+	if (!(Res== AMF_OK || Res==AMF_ALREADY_INITIALIZED))\
+	{\
+		UE_LOG(LogAVEncoder, Error, TEXT("`" #AMF_call "` failed with error code: %d"), Res);\
+		/*check(false);*/\
+		return false;\
+	}\
+}
+
+#define CHECK_AMF_NORET(AMF_call)\
+{\
+	AMF_RESULT Res = AMF_call;\
+	if (Res != AMF_OK)\
+	{\
+		UE_LOG(LogAVEncoder, Error, TEXT("`" #AMF_call "` failed with error code: %d"), Res);\
+	}\
+}
+
+namespace AVEncoder
+{
+
+namespace {
+	// enumerates all available properties of AMFPropertyStorage interface and logs their name,
+	// current and default values and other info
+	inline bool LogAmfPropertyStorage(amf::AMFPropertyStorageEx* PropertyStorage)
+	{
+		SIZE_T NumProps = PropertyStorage->GetPropertiesInfoCount();
+		for (int i = 0; i != NumProps; ++i)
+		{
+			const amf::AMFPropertyInfo* Info;
+			CHECK_AMF_RET(PropertyStorage->GetPropertyInfo(i, &Info));
+
+			if (Info->accessType != amf::AMF_PROPERTY_ACCESS_PRIVATE)
+			{
+				amf::AMFVariant Value;
+				CHECK_AMF_RET(PropertyStorage->GetProperty(Info->name, &Value));
+
+				FString EnumDesc;
+				if (Info->pEnumDescription)
+				{
+					int j = 0;
+					for (; /*j != Value.ToInt32()*/; ++j)
+					{
+						if (Info->pEnumDescription[j].value == Value.ToInt32())
+						{
+							break;
+						}
+					}
+					EnumDesc = TEXT(" ") + FString(Info->pEnumDescription[j].name);
+				}
+
+				UE_LOG(LogAVEncoder, Log, TEXT("Prop %s (%s): value: %s%s, default value: %s (%s - %s), access: %d"),
+					Info->name,
+					Info->desc,
+					Value.ToWString().c_str(),
+					*EnumDesc,
+					amf::AMFVariant{ Info->defaultValue }.ToWString().c_str(),
+					amf::AMFVariant{ Info->minValue }.ToWString().c_str(),
+					amf::AMFVariant{ Info->maxValue }.ToWString().c_str(),
+					static_cast<int>(Info->accessType));
+			}
+			else
+			{
+				UE_LOG(LogAVEncoder, VeryVerbose, TEXT("Prop: %s (%s) - PRIVATE"), Info->name, Info->desc);
+			}
+		}
+
+		return true;
+	}
+
+
+	const TCHAR* ToString(AMF_VIDEO_ENCODER_OUTPUT_DATA_TYPE_ENUM PicType)
+	{
+		switch (PicType)
+		{
+		case AMF_VIDEO_ENCODER_OUTPUT_DATA_TYPE_IDR:
+			return TEXT("AMF_VIDEO_ENCODER_OUTPUT_DATA_TYPE_IDR");
+		case AMF_VIDEO_ENCODER_OUTPUT_DATA_TYPE_I:
+			return TEXT("AMF_VIDEO_ENCODER_OUTPUT_DATA_TYPE_I");
+		case AMF_VIDEO_ENCODER_OUTPUT_DATA_TYPE_P:
+			return TEXT("AMF_VIDEO_ENCODER_OUTPUT_DATA_TYPE_P");
+		case AMF_VIDEO_ENCODER_OUTPUT_DATA_TYPE_B:
+			return TEXT("AMF_VIDEO_ENCODER_OUTPUT_DATA_TYPE_B");
+		default:
+			checkNoEntry();
+			return TEXT("Unknown");
+		}
+	}
+
+	AMF_VIDEO_ENCODER_RATE_CONTROL_METHOD_ENUM ToAmfRcMode(FH264Settings::ERateControlMode RcMode)
+	{
+		switch(RcMode)
+		{
+		case FH264Settings::ERateControlMode::ConstQP:
+			return AMF_VIDEO_ENCODER_RATE_CONTROL_METHOD_CONSTANT_QP;
+		case FH264Settings::ERateControlMode::VBR:
+			return AMF_VIDEO_ENCODER_RATE_CONTROL_METHOD_LATENCY_CONSTRAINED_VBR;
+		case FH264Settings::ERateControlMode::CBR:
+			return AMF_VIDEO_ENCODER_RATE_CONTROL_METHOD_CBR;
+		default:
+			UE_LOG(LogAVEncoder, Error, TEXT("Invalid rate control mode (%d) for nvenc"), (int)RcMode);
+			return AMF_VIDEO_ENCODER_RATE_CONTROL_METHOD_CBR;
+		}
+	};
+
+} // anonymous namespace
+
+class FAmfVideoEncoder : public FVideoEncoder
+{
+public:
+	FAmfVideoEncoder();
+	~FAmfVideoEncoder();
+
+private:
+
+	//
+	// FVideoEncoder interface
+	//
+	const TCHAR* GetName() const override;
+	const TCHAR* GetType() const override;
+	bool Initialize(const FVideoEncoderConfig& InConfig) override;
+	void Shutdown() override;
+	bool CopyTexture(FTexture2DRHIRef Texture, FTimespan CaptureTs, FTimespan Duration, FBufferId& OutBufferId, FIntPoint Resolution = {0, 0}) override;
+	void Drop(FBufferId BufferId) override;
+	void Encode(FBufferId BufferId, bool bForceKeyFrame, uint32 Bitrate = 0, TUniquePtr<AVEncoder::FEncoderVideoFrameCookie> Cookie = nullptr) override;
+	FVideoEncoderConfig GetConfig() const override;
+	bool SetBitrate(uint32 Bitrate) override;
+	bool SetFramerate(uint32 Framerate) override;
+	bool SetParameter(const FString& Parameter, const FString& Value) override;
+
+	struct FInputFrame
+	{
+		FInputFrame() {}
+		UE_NONCOPYABLE(FInputFrame);
+		FTexture2DRHIRef Texture;
+		FTimespan CaptureTs;
+		FTimespan Duration;
+		bool bForceKeyFrame = false;
+	};
+
+	struct FOutputFrame
+	{
+		FOutputFrame() {}
+		UE_NONCOPYABLE(FOutputFrame);
+		amf::AMFDataPtr EncodedData;
+		TUniquePtr<FEncoderVideoFrameCookie> Cookie;
+	};
+
+	enum class EFrameState
+	{
+		Free,
+		Capturing,
+		Captured,
+		EncoderFailed,
+		Encoding
+	};
+
+	struct FFrame
+	{
+		const FBufferId Id = 0;
+		TAtomic<EFrameState> State = EFrameState::Free;
+		FInputFrame InputFrame;
+		FOutputFrame OutputFrame;
+		uint64 FrameIdx = 0;
+
+		// Some timestamps to track how long a frame spends in each step
+		FTimespan CopyBufferStartTs;
+		FTimespan CopyBufferFinishTs;
+		FTimespan EncodingStartTs;
+		FTimespan EncodingFinishTs;
+	};
+
+	bool bInitialized = false;
+	amf_handle DllHandle = nullptr;
+	amf::AMFFactory* AmfFactory = nullptr;
+	amf::AMFContextPtr AmfContext;
+	amf::AMFComponentPtr AmfEncoder;
+	FVideoEncoderConfig Config;
+	FH264Settings ConfigH264;
+
+	uint32 CapturedFrameCount = 0; // of captured, not encoded frames
+	static const uint32 NumBufferedFrames = 3;
+	FFrame BufferedFrames[NumBufferedFrames];
+
+	TQueue<FFrame*> EncodingQueue;
+
+	// Used to make sure we don't have a race condition trying to access a deleted "this" captured
+	// in the render command lambda sent to the render thread from EncoderCheckLoop
+	static FThreadSafeCounter ImplCounter;
+
+	FRHICOMMAND_MACRO(FRHISubmitFrameToEncoder)
+	{
+		FAmfVideoEncoder* Encoder;
+		FFrame* Frame;
+		FRHISubmitFrameToEncoder(FAmfVideoEncoder* InEncoder, FFrame* InFrame)
+			: Encoder(InEncoder), Frame(InFrame)
+		{
+		}
+
+		void Execute(FRHICommandListBase& CmdList)
+		{
+			Encoder->SubmitFrameToEncoder(*Frame);
+		}
+	};
+
+	void ResetFrameInputBuffer(FFrame& Frame, FIntPoint Resolution);
+	bool ProcessOutput();
+	void UpdateRes(FFrame& Frame, FIntPoint Resolution);
+	void EncodeFrameInRenderingThread(FFrame& Frame, uint32 Bitrate);
+	bool SubmitFrameToEncoder(FFrame& Frame);
+	void UpdateEncoderConfig(FIntPoint Resolution, uint32 Bitrate);
+	void HandleDroppedFrame(FFrame& Frame);
+	void HandleEncodedFrame(FFrame& Frame);
+};
+
+//////////////////////////////////////////////////////////////////////////
+// FAmfVideoEncoder
+//////////////////////////////////////////////////////////////////////////
+
+FThreadSafeCounter FAmfVideoEncoder::ImplCounter(0);
+
+FAmfVideoEncoder::FAmfVideoEncoder()
+{
+}
+
+FAmfVideoEncoder::~FAmfVideoEncoder()
+{
+	if (DllHandle)
+	{
+		UE_LOG(LogAVEncoder, Fatal, TEXT("FAmfVideoEncoder Shutdown not called before destruction."));
+	}
+}
+
+const TCHAR* FAmfVideoEncoder::GetName() const
+{
+	return TEXT("h264.amf");
+}
+
+const TCHAR* FAmfVideoEncoder::GetType() const
+{
+	return TEXT("h264");
+}
+
+bool FAmfVideoEncoder::Initialize(const FVideoEncoderConfig& InConfig)
+{
+	check(!bInitialized);
+
+	Config = InConfig;
+	ConfigH264 = {};
+	ReadH264Settings(Config.Options, ConfigH264);
+
+	UE_LOG(LogAVEncoder, Log, TEXT("FPixelStreamingAmfVideoEncoder initialization with %dx%d, %d FPS")
+		, Config.Width, Config.Height, Config.Framerate);
+
+	DllHandle = FPlatformProcess::GetDllHandle(AMF_DLL_NAME);
+	if (DllHandle == nullptr)
+	{
+		return false;
+	}
+
+	AMFInit_Fn AmfInitFn = (AMFInit_Fn)FPlatformProcess::GetDllExport((HMODULE)DllHandle, TEXT(AMF_INIT_FUNCTION_NAME));
+	if (AmfInitFn == nullptr)
+	{
+		return false;
+	}
+	CHECK_AMF_RET(AmfInitFn(AMF_FULL_VERSION, &AmfFactory));
+
+	AMFQueryVersion_Fn AmfVersionFun = (AMFQueryVersion_Fn)FPlatformProcess::GetDllExport((HMODULE)DllHandle, TEXT(AMF_QUERY_VERSION_FUNCTION_NAME));
+	if (AmfVersionFun == nullptr)
+	{
+		return false;
+	}
+	uint64 AmfVersion = 0;
+	AmfVersionFun(&AmfVersion);
+
+	FString RHIName = GDynamicRHI->GetName();
+	if (RHIName != TEXT("D3D11"))
+	{
+		UE_LOG(LogAVEncoder, Fatal, TEXT("AMF not supported with a %s renderer"), *RHIName);
+	}
+
+	ID3D11Device* DxDevice = static_cast<ID3D11Device*>(GDynamicRHI->RHIGetNativeDevice());
+
+	CHECK_AMF_RET(AmfFactory->CreateContext(&AmfContext));
+
+	checkf(DxDevice != nullptr, TEXT("Cannot initialize NvEnc with invalid device"));
+	CHECK_AMF_RET(AmfContext->InitDX11(DxDevice));
+
+	CHECK_AMF_RET(AmfFactory->CreateComponent(AmfContext, AMFVideoEncoderVCE_AVC, &AmfEncoder));
+
+	if (Config.Preset == FVideoEncoderConfig::EPreset::LowLatency)
+	{
+		CHECK_AMF_RET(AmfEncoder->SetProperty(AMF_VIDEO_ENCODER_USAGE, AMF_VIDEO_ENCODER_USAGE_ULTRA_LOW_LATENCY));
+		CHECK_AMF_RET(AmfEncoder->SetProperty(AMF_VIDEO_ENCODER_LOWLATENCY_MODE, true));
+		CHECK_AMF_RET(AmfEncoder->SetProperty(AMF_VIDEO_ENCODER_PROFILE, AMF_VIDEO_ENCODER_PROFILE_BASELINE));
+		CHECK_AMF_RET(AmfEncoder->SetProperty(AMF_VIDEO_ENCODER_QUALITY_PRESET, AMF_VIDEO_ENCODER_QUALITY_PRESET_BALANCED));
+	}
+	else
+	{
+		CHECK_AMF_RET(AmfEncoder->SetProperty(AMF_VIDEO_ENCODER_USAGE, AMF_VIDEO_ENCODER_USAGE_TRANSCONDING));
+		CHECK_AMF_RET(AmfEncoder->SetProperty(AMF_VIDEO_ENCODER_PROFILE, AMF_VIDEO_ENCODER_PROFILE_MAIN));
+		CHECK_AMF_RET(AmfEncoder->SetProperty(AMF_VIDEO_ENCODER_QUALITY_PRESET, AMF_VIDEO_ENCODER_QUALITY_PRESET_QUALITY));
+	}
+
+	//CHECK_AMF_RET(AmfEncoder->SetProperty(AMF_VIDEO_ENCODER_B_PIC_PATTERN, 0));
+	CHECK_AMF_RET(AmfEncoder->SetProperty(AMF_VIDEO_ENCODER_TARGET_BITRATE, Config.Bitrate));
+	CHECK_AMF_RET(AmfEncoder->SetProperty(AMF_VIDEO_ENCODER_FRAMESIZE, ::AMFConstructSize(Config.Width, Config.Height)));
+	CHECK_AMF_RET(AmfEncoder->SetProperty(AMF_VIDEO_ENCODER_ASPECT_RATIO, ::AMFConstructRatio(Config.Width, Config.Height)));
+	CHECK_AMF_RET(AmfEncoder->SetProperty(AMF_VIDEO_ENCODER_FRAMERATE, ::AMFConstructRate(Config.Framerate, 1)));
+
+	// generate key-frames every second, useful for seeking in resulting .mp4 and keeping recording ring buffer
+	// of second-precise duration
+	uint64 IdrPeriod = Config.Framerate;
+	CHECK_AMF_RET(AmfEncoder->SetProperty(AMF_VIDEO_ENCODER_IDR_PERIOD, IdrPeriod));
+	// insert SPS/PPS before every key-frame. .mp4 file video stream must start from SPS/PPS. Their size is
+	// negligible so having them before every key-frame is not an issue but they presence simplifies
+	// implementation significantly. Otherwise we would extract SPS/PPS from the first key-frame and store
+	// them manually at the beginning of resulting .mp4 file
+	CHECK_AMF_RET(AmfEncoder->SetProperty(AMF_VIDEO_ENCODER_HEADER_INSERTION_SPACING, IdrPeriod));
+
+	CHECK_AMF_RET(AmfEncoder->Init(amf::AMF_SURFACE_RGBA, Config.Width, Config.Height));
+
+	// #AVENCODER : This doesn't seem to be working. It fails with code 3 (AMF_ACCESS_DENIED) or gives a warning that has bee initialized already,
+	// depending if it's called before the Init (gives error 3), or after the Init (giveswarning).
+	CHECK_AMF_RET(AmfEncoder->SetProperty(AMF_VIDEO_ENCODER_RATE_CONTROL_METHOD, ToAmfRcMode(ConfigH264.RcMode)));
+
+	CHECK_AMF_RET(AmfEncoder->SetProperty(AMF_VIDEO_ENCODER_MIN_QP, ConfigH264.QP));
+	CHECK_AMF_RET(AmfEncoder->SetProperty(AMF_VIDEO_ENCODER_QP_I, ConfigH264.QP));
+	CHECK_AMF_RET(AmfEncoder->SetProperty(AMF_VIDEO_ENCODER_QP_P, ConfigH264.QP));
+	// #AVENCODER : This doesn't seem to be working. It fails with code 3 (AMF_ACCESS_DENIED) or gives a warning that has bee initialized already,
+	// depending if it's called before the Init (gives error 3), or after the Init (giveswarning).
+	CHECK_AMF_RET(AmfEncoder->SetProperty(AMF_VIDEO_ENCODER_QP_B, ConfigH264.QP));
+
+	//LogAmfPropertyStorage(AmfEncoder);
+
+	FBufferId Id = 0;
+	for (FFrame& Frame : BufferedFrames)
+	{
+		// We keep `Id` as const, because it not supposed to change at all once initialized
+		*(const_cast<FBufferId*>(&Frame.Id)) = Id++;
+		ResetFrameInputBuffer(Frame, { (int)Config.Width, (int)Config.Height });
+	}
+	UE_LOG(LogAVEncoder, Log, TEXT("AMF H.264 encoder initialised, v.0x%X"), AmfVersion);
+	bInitialized = true;
+	return true;
+}
+
+void FAmfVideoEncoder::Shutdown()
+{
+	// BufferedFrames keep references to AMFData, we need to release them before destroying AMF
+	for (FFrame& Frame : BufferedFrames)
+	{
+		Frame.OutputFrame.EncodedData = nullptr;
+	}
+
+	// Cleanup in this order
+	if (AmfEncoder)
+	{
+		AmfEncoder->Terminate();
+		AmfEncoder = nullptr;
+	}
+	if (AmfContext)
+	{
+		AmfContext->Terminate();
+		AmfContext = nullptr;
+	}
+	AmfFactory = nullptr;
+	if (DllHandle)
+	{
+		FPlatformProcess::FreeDllHandle(DllHandle);
+		DllHandle = nullptr;
+	}
+}
+
+bool FAmfVideoEncoder::CopyTexture(FTexture2DRHIRef Texture, FTimespan CaptureTs, FTimespan Duration, FBufferId& OutBufferId, FIntPoint Resolution)
+{
+	check(IsInRenderingThread());
+
+	// First process output, to free slots we can use for the new request
+	if (!ProcessOutput())
+	{
+		return false;
+	}
+
+	// Find a free slot we can use
+	FFrame* Frame = nullptr;
+	for (FFrame& Slot : BufferedFrames)
+	{
+		if (Slot.State.Load() == EFrameState::Free)
+		{
+			Frame = &Slot;
+			OutBufferId = Slot.Id;
+			break;
+		}
+	}
+
+	if (!Frame)
+	{
+		UE_LOG(LogAVEncoder, Verbose, TEXT("Frame dropped because Amf queue is full"));
+		return false;
+	}
+
+	Frame->FrameIdx = CapturedFrameCount++;
+	Frame->InputFrame.CaptureTs = CaptureTs;
+	Frame->InputFrame.Duration = Duration;
+	Frame->CopyBufferStartTs = FTimespan::FromSeconds(FPlatformTime::Seconds());
+
+	{
+		UpdateRes(*Frame, Resolution.Size() ? Resolution : Texture->GetSizeXY());
+		CopyTextureImpl(Texture, Frame->InputFrame.Texture, nullptr);
+	}
+
+	UE_LOG(LogAVEncoder, Verbose, TEXT("Buffer #%d (%d) captured"), Frame->FrameIdx, OutBufferId);
+	Frame->State = EFrameState::Capturing;
+
+	return true;
+}
+
+void FAmfVideoEncoder::Drop(FBufferId BufferId)
+{
+	FFrame& Frame = BufferedFrames[BufferId];
+
+	{
+		EFrameState State = Frame.State.Load();
+		checkf(State == EFrameState::Capturing, TEXT("Buffer %d: Expected state %d, found %d")
+			, BufferId, (int)EFrameState::Capturing, (int)State);
+	}
+
+	Frame.State = EFrameState::Free;
+
+	UE_LOG(LogAVEncoder, Log, TEXT("Buffer #%d (%d) dropped"), BufferedFrames[BufferId].FrameIdx, BufferId);
+}
+
+
+void FAmfVideoEncoder::Encode(FBufferId BufferId, bool bForceKeyFrame, uint32 Bitrate, TUniquePtr<AVEncoder::FEncoderVideoFrameCookie> Cookie)
+{
+	FFrame& Frame = BufferedFrames[BufferId];
+
+	{
+		EFrameState State = Frame.State.Load();
+		checkf(State == EFrameState::Capturing, TEXT("Buffer %d : Expected state %d, but found %d"), BufferId, (int)EFrameState::Captured, (int)State);
+	}
+
+	Frame.State = EFrameState::Captured;
+	Frame.CopyBufferFinishTs = FTimespan::FromSeconds(FPlatformTime::Seconds());
+	Frame.InputFrame.bForceKeyFrame = bForceKeyFrame;
+	Frame.OutputFrame.Cookie = MoveTemp(Cookie);
+
+	int32 CurrImplCounter = ImplCounter.GetValue();
+	ENQUEUE_RENDER_COMMAND(AmfEncEncodeFrame)(
+	[this, &Frame, BufferId, Bitrate, CurrImplCounter](FRHICommandListImmediate& RHICmdList)
+	{
+		if (CurrImplCounter != ImplCounter.GetValue()) // Check if the "this" we captured is still valid
+			return;
+
+		EncodeFrameInRenderingThread(Frame, Bitrate);
+
+		UE_LOG(LogAVEncoder, VeryVerbose, TEXT("Buffer #%d (%d), ts %lld started encoding"), Frame.FrameIdx, BufferId, Frame.InputFrame.CaptureTs.GetTicks());
+	});
+}
+
+FVideoEncoderConfig FAmfVideoEncoder::GetConfig() const
+{
+	return Config;
+}
+bool FAmfVideoEncoder::SetBitrate(uint32 Bitrate)
+{
+	Config.Bitrate = Bitrate;
+	return true;
+}
+
+bool FAmfVideoEncoder::SetFramerate(uint32 Framerate)
+{
+	Config.Framerate = Framerate;
+	return true;
+}
+
+bool FAmfVideoEncoder::SetParameter(const FString& Parameter, const FString& Value)
+{
+	return ReadH264Setting(Parameter, Value, ConfigH264);
+}
+
+void FAmfVideoEncoder::ResetFrameInputBuffer(FFrame& Frame, FIntPoint Resolution)
+{
+	Frame.InputFrame.Texture.SafeRelease();
+
+	// Make sure format used here is compatible with AMF_SURFACE_FORMAT specified in encoder Init() function.
+	FRHIResourceCreateInfo CreateInfo;
+	Frame.InputFrame.Texture = RHICreateTexture2D(Resolution.X, Resolution.Y, EPixelFormat::PF_R8G8B8A8, 1, 1, TexCreate_RenderTargetable, CreateInfo);
+}
+
+bool FAmfVideoEncoder::ProcessOutput()
+{
+	check(IsInRenderingThread());
+
+	while(true)
+	{
+
+		// Drop any failed submits
+		// If for some reason this frame failed to be submitted to AMF (e.g: call to AmfEncoder->SubmitInput failed due to AMF_INPUT_FULL), then we need to tell webrtc to drop it
+		while(!EncodingQueue.IsEmpty() && (*EncodingQueue.Peek())->State == EFrameState::EncoderFailed)
+		{
+			FFrame* Frame;
+			EncodingQueue.Dequeue(Frame);
+			HandleDroppedFrame(*Frame);
+		}
+
+		amf::AMFDataPtr EncodedData;
+		AMF_RESULT Ret;
+		{
+			SCOPE_CYCLE_COUNTER(STAT_Amf_QueryOutput);
+			Ret = AmfEncoder->QueryOutput(&EncodedData);
+		}
+
+		// If AMF says a frame finished encoding, then we must have one in EncodingQueue by design
+		if (Ret == AMF_OK && EncodedData != nullptr)
+		{
+			FFrame* Frame;
+			verify(EncodingQueue.Dequeue(Frame));
+			Frame->OutputFrame.EncodedData = EncodedData;
+			HandleEncodedFrame(*Frame);
+		}
+		else if (Ret == AMF_REPEAT)
+		{
+			break; // not ready yet
+		}
+		else
+		{
+			UE_LOG(LogAVEncoder, Error, TEXT("Failed to query AMF H.264 Encoder output: %d, %p"), Ret, EncodedData.GetPtr());
+			return false;
+		}
+	}
+
+	return true;
+}
+
+
+void FAmfVideoEncoder::HandleDroppedFrame(FFrame& Frame)
+{
+	check(Frame.State==EFrameState::EncoderFailed);
+
+	FOutputFrame& OutputFrame = Frame.OutputFrame;
+
+	FAVPacket Packet(EPacketType::Video);
+	Frame.EncodingFinishTs = FTimespan::FromSeconds(FPlatformTime::Seconds());
+	Packet.Timestamp = Frame.InputFrame.CaptureTs;
+	Packet.Duration = Frame.InputFrame.Duration;
+	Packet.Timings.EncodeStartTs = Frame.EncodingStartTs;
+	Packet.Timings.EncodeFinishTs = Frame.EncodingFinishTs;
+
+	UE_LOG(LogAVEncoder, VeryVerbose, TEXT("dropping frame with ts %lld due to encoder failure") , Packet.Timestamp.GetTicks());
+
+	{
+		SCOPE_CYCLE_COUNTER(STAT_Amf_OnEncodedVideoFrameCallback);
+		OnEncodedVideoFrame(Packet, MoveTemp(Frame.OutputFrame.Cookie));
+	}
+
+	Frame.State = EFrameState::Free;
+}
+
+void FAmfVideoEncoder::HandleEncodedFrame(FFrame& Frame)
+{
+	check(Frame.State==EFrameState::Encoding);
+
+	FOutputFrame& OutputFrame = Frame.OutputFrame;
+
+	FAVPacket Packet(EPacketType::Video);
+
+	//
+	// Query for buffer interface
+	//
+	amf::AMFBufferPtr EncodedBuffer(Frame.OutputFrame.EncodedData);
+	void* EncodedBufferPtr = EncodedBuffer->GetNative();
+	size_t EncodedBufferSize = EncodedBuffer->GetSize();
+
+	//
+	// Check if it's a keyframe
+	//
+	amf_int64 PicType;
+	CHECK_AMF_NORET(EncodedBuffer->GetProperty(AMF_VIDEO_ENCODER_OUTPUT_DATA_TYPE, &PicType));
+	bool bKeyFrame = PicType == AMF_VIDEO_ENCODER_OUTPUT_DATA_TYPE_IDR ? true : false;
+	checkf(bKeyFrame == true || Frame.InputFrame.bForceKeyFrame == false, TEXT("key frame requested by webrtc but not provided by Amf"));
+
+	Frame.EncodingFinishTs = FTimespan::FromSeconds(FPlatformTime::Seconds());
+	Packet.Timestamp = Frame.InputFrame.CaptureTs;
+	Packet.Duration = Frame.InputFrame.Duration;
+	Packet.Video.bKeyFrame = bKeyFrame;
+	Packet.Video.Width = Frame.InputFrame.Texture->GetSizeX();
+	Packet.Video.Height = Frame.InputFrame.Texture->GetSizeY();
+	Packet.Video.Framerate = Config.Framerate;
+	// #AVENCODER : How to get frame average QP from Amf ?
+	Packet.Video.FrameAvgQP = 20;
+	Packet.Data = TArray<uint8>(reinterpret_cast<const uint8*>(EncodedBufferPtr), EncodedBufferSize);
+	Packet.Timings.EncodeStartTs = Frame.EncodingStartTs;
+	Packet.Timings.EncodeFinishTs = Frame.EncodingFinishTs;
+
+	// Release AMF's buffer
+	Frame.OutputFrame.EncodedData = nullptr;
+
+	UE_LOG(LogAVEncoder, VeryVerbose, TEXT("encoded %s ts %lld, %d bytes")
+		, ToString((AMF_VIDEO_ENCODER_OUTPUT_DATA_TYPE_ENUM)PicType)
+		, Packet.Timestamp.GetTicks()
+		, (int)Packet.Data.Num());
+
+	{
+		SCOPE_CYCLE_COUNTER(STAT_Amf_OnEncodedVideoFrameCallback);
+		OnEncodedVideoFrame(Packet, MoveTemp(Frame.OutputFrame.Cookie));
+	}
+
+	Frame.State = EFrameState::Free;
+}
+
+
+void FAmfVideoEncoder::UpdateRes(FFrame& Frame, FIntPoint Resolution)
+{
+	check(IsInRenderingThread());
+
+	FInputFrame& InputFrame = Frame.InputFrame;
+
+	// check if target resolution matches our currently allocated `InputFrame.Texture` resolution
+	if (InputFrame.Texture->GetSizeX() == Resolution.X && InputFrame.Texture->GetSizeY() == Resolution.Y)
+	{
+		return;
+	}
+
+	ResetFrameInputBuffer(Frame, Resolution);
+}
+
+void FAmfVideoEncoder::EncodeFrameInRenderingThread(FFrame& Frame, uint32 Bitrate)
+{
+	check(IsInRenderingThread());
+	check(Frame.State==EFrameState::Captured);
+
+	UpdateEncoderConfig(Frame.InputFrame.Texture->GetSizeXY(), Bitrate);
+
+	//
+	// Process the new input
+	{
+		FRHICommandList& RHICmdList = FRHICommandListExecutor::GetImmediateCommandList();
+		if (RHICmdList.Bypass())
+		{
+			FRHISubmitFrameToEncoder Command(this, &Frame);
+			Command.Execute(RHICmdList);
+		}
+		else
+		{
+			ALLOC_COMMAND_CL(RHICmdList, FRHISubmitFrameToEncoder)(this, &Frame);
+		}
+	}
+}
+
+bool FAmfVideoEncoder::SubmitFrameToEncoder(FFrame& Frame)
+{
+	SCOPE_CYCLE_COUNTER(STAT_Amf_SubmitFrameToEncoder);
+	check(Frame.State==EFrameState::Captured);
+
+	Frame.EncodingStartTs = FTimespan::FromSeconds(FPlatformTime::Seconds());
+
+	amf::AMFSurfacePtr AmfSurfaceIn;
+	ID3D11Texture2D* ResolvedBackBufferDX11 = static_cast<ID3D11Texture2D*>(GetD3D11TextureFromRHITexture(Frame.InputFrame.Texture)->GetResource());
+	CHECK_AMF_RET(AmfContext->CreateSurfaceFromDX11Native(ResolvedBackBufferDX11, &AmfSurfaceIn, nullptr));
+
+	if (Frame.InputFrame.bForceKeyFrame)
+	{
+		CHECK_AMF_RET(AmfSurfaceIn->SetProperty(AMF_VIDEO_ENCODER_FORCE_PICTURE_TYPE, AMF_VIDEO_ENCODER_PICTURE_TYPE_IDR));
+	}
+
+	{
+		// if `-d3ddebug` is enabled `SubmitInput` crashes with DX11 error, see output window
+		// we believe it's an internal AMF shader problem so we disable those errors explicitly, otherwise
+		// DX Debug Layer can't be used at all
+		FScopeDisabledDxDebugErrors Errors({
+			D3D11_MESSAGE_ID_DEVICE_UNORDEREDACCESSVIEW_RETURN_TYPE_MISMATCH,
+			D3D11_MESSAGE_ID_DEVICE_CSSETUNORDEREDACCESSVIEWS_TOOMANYVIEWS
+		});
+
+		{
+			SCOPE_CYCLE_COUNTER(STAT_Amf_SubmitInput);
+			//CHECK_AMF_RET(AmfEncoder->SubmitInput(AmfSurfaceIn));
+			AMF_RESULT Res =  AmfEncoder->SubmitInput(AmfSurfaceIn);
+			if (Res == AMF_OK)
+			{
+				Frame.State = EFrameState::Encoding;
+			}
+			else
+			{
+				UE_LOG(LogAVEncoder, Error, TEXT("AmfEncoder->SubmitInput failed failed with error code: %d"), Res);
+				Frame.State = EFrameState::EncoderFailed;
+			}
+			EncodingQueue.Enqueue(&Frame);
+		}
+	}
+
+	return true;
+}
+
+
+void FAmfVideoEncoder::UpdateEncoderConfig(FIntPoint Resolution, uint32 Bitrate)
+{
+	check(IsInRenderingThread());
+
+	//
+	// If an explicit Bitrate was specified, use that one, if not, use the one
+	// from the Config struct
+	//
+	uint32 AmfBitrate = 0;
+	CHECK_AMF_NORET(AmfEncoder->GetProperty(AMF_VIDEO_ENCODER_TARGET_BITRATE, &AmfBitrate));
+	if (Bitrate)
+	{
+		if (AmfBitrate != Bitrate)
+		{
+			UE_LOG(LogAVEncoder, Verbose, TEXT("Setting AMF's bitrate to %u"), Bitrate);
+			CHECK_AMF_NORET(AmfEncoder->SetProperty(AMF_VIDEO_ENCODER_TARGET_BITRATE, Bitrate));
+			Config.Bitrate = Bitrate;
+		}
+	}
+	else if (AmfBitrate != Config.Bitrate)
+	{
+		UE_LOG(LogAVEncoder, Verbose, TEXT("Setting AMF's bitrate to %u"), Config.Bitrate);
+		CHECK_AMF_NORET(AmfEncoder->SetProperty(AMF_VIDEO_ENCODER_TARGET_BITRATE, Config.Bitrate));
+	}
+
+	//
+	// Update QP if required
+	//
+	uint32 AmfQP = 0;
+	CHECK_AMF_NORET(AmfEncoder->GetProperty(AMF_VIDEO_ENCODER_MIN_QP, &AmfQP));
+	if (AmfQP != ConfigH264.QP)
+	{
+		UE_LOG(LogAVEncoder, Verbose, TEXT("Setting AMF's MIN_QP/QP_I/QP_P/QP_B to %u"), ConfigH264.QP);
+		CHECK_AMF_NORET(AmfEncoder->SetProperty(AMF_VIDEO_ENCODER_MIN_QP, ConfigH264.QP));
+		CHECK_AMF_NORET(AmfEncoder->SetProperty(AMF_VIDEO_ENCODER_QP_I  , ConfigH264.QP));
+		CHECK_AMF_NORET(AmfEncoder->SetProperty(AMF_VIDEO_ENCODER_QP_P  , ConfigH264.QP));
+		CHECK_AMF_NORET(AmfEncoder->SetProperty(AMF_VIDEO_ENCODER_QP_B  , ConfigH264.QP));
+	}
+
+	//
+	// Update resolution if required
+	//
+	if (Resolution.Size())
+	{
+		AMFSize AmfResolution;
+		CHECK_AMF_NORET(AmfEncoder->GetProperty(AMF_VIDEO_ENCODER_FRAMESIZE, &AmfResolution));
+		if (Resolution != FIntPoint(AmfResolution.width, AmfResolution.height))
+		{
+			UE_LOG(LogAVEncoder, Verbose, TEXT("Setting AMF's Resolution to %dx%d"), Resolution.X, Resolution.Y);
+			CHECK_AMF_NORET(AmfEncoder->SetProperty(AMF_VIDEO_ENCODER_FRAMESIZE, ::AMFConstructSize(Resolution.X, Resolution.Y)));
+			CHECK_AMF_NORET(AmfEncoder->SetProperty(AMF_VIDEO_ENCODER_ASPECT_RATIO, ::AMFConstructRatio(Resolution.X, Resolution.Y)));
+		}
+	}
+
+	//
+	// Update framerate if necessary
+	// 
+	AMFRate AmfFramerate;
+	CHECK_AMF_NORET(AmfEncoder->GetProperty(AMF_VIDEO_ENCODER_FRAMERATE, &AmfFramerate));
+	if (AmfFramerate.num != Config.Framerate)
+	{
+		UE_LOG(LogAVEncoder, Verbose, TEXT("Setting AMF's framerate to %d"), Config.Framerate);
+		CHECK_AMF_NORET(AmfEncoder->SetProperty(AMF_VIDEO_ENCODER_FRAMERATE, ::AMFConstructRate(Config.Framerate, 1)));
+		CHECK_AMF_NORET(AmfEncoder->SetProperty(AMF_VIDEO_ENCODER_IDR_PERIOD, Config.Framerate));
+		CHECK_AMF_NORET(AmfEncoder->SetProperty(AMF_VIDEO_ENCODER_HEADER_INSERTION_SPACING, Config.Framerate));
+	}
+
+	//
+	// Implement Rate Control Mode
+	//
+	// #AVENCODER : This is not working. GetProperty AMF_VIDEO_ENCODER_RATE_CONTROL_METHOD fails with error 3 (AMF_ACCESS_DENIED)
+#if 0
+	amf_int64 AmfRcMode = AMF_VIDEO_ENCODER_RATE_CONTROL_METHOD_CBR;
+	amf_int64 CfgRcMode = ToAmfRcMode(ConfigH264.RcMode);
+	CHECK_AMF_NORET(AmfEncoder->GetProperty(AMF_VIDEO_ENCODER_RATE_CONTROL_METHOD, &AmfRcMode));
+	if (AmfRcMode != CfgRcMode)
+	{
+		UE_LOG(LogAVEncoder, Verbose, TEXT("Setting AMF's control rate method to %d"), int(CfgRcMode));
+		CHECK_AMF_NORET(AmfEncoder->SetProperty(AMF_VIDEO_ENCODER_RATE_CONTROL_METHOD, CfgRcMode));
+	}
+#endif
+}
+
+//////////////////////////////////////////////////////////////////////////
+// FAmfVideoEncoderFactory
+//////////////////////////////////////////////////////////////////////////
+
+FAmfVideoEncoderFactory::FAmfVideoEncoderFactory()
+{
+}
+
+FAmfVideoEncoderFactory::~FAmfVideoEncoderFactory()
+{
+}
+
+const TCHAR* FAmfVideoEncoderFactory::GetName() const
+{
+	return TEXT("amf");
+}
+
+TArray<FString> FAmfVideoEncoderFactory::GetSupportedCodecs() const
+{
+	TArray<FString> Codecs;
+
+	if (!IsRHIDeviceAMD())
+	{
+		UE_LOG(LogAVEncoder, Log, TEXT("No AMF available because no AMD card found"));
+		return Codecs;
+	}
+
+	void* Handle = FPlatformProcess::GetDllHandle(AMF_DLL_NAME);
+	if (Handle == nullptr)
+	{
+		UE_LOG(LogAVEncoder, Error, TEXT("AMD card found, but no AMF DLL installed."));
+		return Codecs;
+	}
+	else
+	{
+		FPlatformProcess::FreeDllHandle(Handle);
+	}
+
+	Codecs.Add(TEXT("h264"));
+	return Codecs;
+}
+
+TUniquePtr<FVideoEncoder> FAmfVideoEncoderFactory::CreateEncoder(const FString& Codec)
+{
+	if (Codec=="h264")
+	{
+		return TUniquePtr<FVideoEncoder>(new FAmfVideoEncoder());
+	}
+	else
+	{
+		UE_LOG(LogAVEncoder, Error, TEXT("FAmfVideoEncoderFactory doesn't support the %s codec"), *Codec);
+		return nullptr;
+	}
+}
+
+} // namespace AVEncoder
+
diff --git a/Engine/Source/Runtime/AVEncoder/Private/Microsoft/Windows/AmfVideoEncoder.h b/Engine/Source/Runtime/AVEncoder/Private/Microsoft/Windows/AmfVideoEncoder.h
new file mode 100644
--- /dev/null
+++ b/Engine/Source/Runtime/AVEncoder/Private/Microsoft/Windows/AmfVideoEncoder.h
@@ -0,0 +1,26 @@
+// Copyright 1998-2019 Epic Games, Inc. All Rights Reserved.
+
+#pragma once
+
+#include "AVEncoder.h"
+#include "AVEncoderCommon.h"
+
+namespace AVEncoder
+{
+
+/**
+ * Factory for AMD's Amf
+ */
+class FAmfVideoEncoderFactory : public FVideoEncoderFactory
+{
+public:
+	FAmfVideoEncoderFactory();
+	~FAmfVideoEncoderFactory() override;
+	const TCHAR* GetName() const override;
+	TArray<FString> GetSupportedCodecs() const override;
+	TUniquePtr<FVideoEncoder> CreateEncoder(const FString& Codec) override;
+private:
+};
+
+}
+
