UE-74862 : Reusable media encoder framework implementation

diff --git a/Engine/Source/Runtime/AVEncoder/Private/Microsoft/Windows/NvVideoEncoder.cpp b/Engine/Source/Runtime/AVEncoder/Private/Microsoft/Windows/NvVideoEncoder.cpp
new file mode 100644
--- /dev/null
+++ b/Engine/Source/Runtime/AVEncoder/Private/Microsoft/Windows/NvVideoEncoder.cpp
@@ -0,0 +1,1288 @@
+// Copyright 1998-2019 Epic Games, Inc. All Rights Reserved.
+
+#include "NvVideoEncoder.h"
+#include "Microsoft/AVEncoderMicrosoftCommon.h"
+
+THIRD_PARTY_INCLUDES_START
+	#include "NvEncoder/nvEncodeAPI.h"
+THIRD_PARTY_INCLUDES_END
+
+// This is mostly to use internally at Epic.
+// Setting this to 1 will collect detailed timings in the `Timings` member array.
+// It will also clear every frame with a solid colour before copying the backbuffer into it.
+#define NVENC_VIDEO_ENCODER_DEBUG 0
+
+#if NVENC_VIDEO_ENCODER_DEBUG
+	#include "ClearQuad.h"
+#endif
+
+#define CHECK_NV_RES(NvCall)\
+{\
+	NVENCSTATUS Res = NvCall;\
+	if (Res != NV_ENC_SUCCESS)\
+	{\
+		check(false);\
+		UE_LOG(LogAVEncoder, Error, TEXT("`" #NvCall "` failed with error code: %d"), Res);\
+		return false;\
+	}\
+}
+
+namespace AVEncoder
+{
+
+namespace
+{
+	constexpr const TCHAR* GetDllName()
+	{
+#if defined PLATFORM_WINDOWS
+#if defined _WIN64
+		return TEXT("nvEncodeAPI64.dll");
+#else
+		return TEXT("nvEncodeAPI.dll");
+#endif
+#elif defined PLATFORM_LINUX
+		return TEXT("libnvidia-encode.so.1");
+#else
+		return TEXT("");
+#endif
+	}
+
+	DECLARE_STATS_GROUP(TEXT("NvEnc"), STATGROUP_NvEncVideoEncoder, STATCAT_Advanced);
+
+	DECLARE_CYCLE_STAT(TEXT("CopyTexture"), STAT_NvEnc_CopyTexture, STATGROUP_NvEncVideoEncoder);
+	DECLARE_CYCLE_STAT(TEXT("SubmitFrameToEncoder"), STAT_NvEnc_SubmitFrameToEncoder, STATGROUP_NvEncVideoEncoder);
+	DECLARE_CYCLE_STAT(TEXT("WaitForEncodeEvent"), STAT_NvEnc_WaitForEncodeEvent, STATGROUP_NvEncVideoEncoder);
+	DECLARE_CYCLE_STAT(TEXT("RetrieveEncodedFrame"), STAT_NvEnc_RetrieveEncodedFrame, STATGROUP_NvEncVideoEncoder);
+	DECLARE_CYCLE_STAT(TEXT("OnEncodedVideoFrameCallback"), STAT_NvEnc_OnEncodedVideoFrameCallback, STATGROUP_NvEncVideoEncoder);
+
+	NV_ENC_PARAMS_RC_MODE ToNvEncRcMode(FH264Settings::ERateControlMode RcMode)
+	{
+		switch(RcMode)
+		{
+		case FH264Settings::ERateControlMode::ConstQP:
+			return NV_ENC_PARAMS_RC_CONSTQP;
+		case FH264Settings::ERateControlMode::VBR:
+			return NV_ENC_PARAMS_RC_VBR;
+		case FH264Settings::ERateControlMode::CBR:
+			return NV_ENC_PARAMS_RC_CBR;
+		default:
+			UE_LOG(LogAVEncoder, Error, TEXT("Invalid rate control mode (%d) for nvenc"), (int)RcMode);
+			return NV_ENC_PARAMS_RC_CBR;
+		}
+	}
+
+	const TCHAR* ToString(NV_ENC_PARAMS_RC_MODE RcMode)
+	{
+		switch (RcMode)
+		{
+		case NV_ENC_PARAMS_RC_CONSTQP:
+			return TEXT("ConstQP");
+		case NV_ENC_PARAMS_RC_VBR:
+			return TEXT("VBR");
+		case NV_ENC_PARAMS_RC_CBR:
+			return TEXT("CBR");
+		case NV_ENC_PARAMS_RC_CBR_LOWDELAY_HQ:
+			return TEXT("CBR_LOWDELAY_HQ");
+		case NV_ENC_PARAMS_RC_CBR_HQ:
+			return TEXT("CBR_HQ");
+		case NV_ENC_PARAMS_RC_VBR_HQ:
+			return TEXT("VBR_HQ");
+		default:
+			checkNoEntry();
+			return TEXT("Unknown");
+		}
+	}
+
+	const TCHAR* ToString(NV_ENC_PIC_TYPE PicType)
+	{
+		switch (PicType)
+		{
+		case NV_ENC_PIC_TYPE_P:
+			return TEXT("NV_ENC_PIC_TYPE_P");
+		case NV_ENC_PIC_TYPE_B:
+			return TEXT("NV_ENC_PIC_TYPE_B");
+		case NV_ENC_PIC_TYPE_I:
+			return TEXT("NV_ENC_PIC_TYPE_I");
+		case NV_ENC_PIC_TYPE_IDR:
+			return TEXT("NV_ENC_PIC_TYPE_IDR");
+		case NV_ENC_PIC_TYPE_BI:
+			return TEXT("NV_ENC_PIC_TYPE_BI");
+		case NV_ENC_PIC_TYPE_SKIPPED:
+			return TEXT("NV_ENC_PIC_TYPE_SKIPPED");
+		case NV_ENC_PIC_TYPE_INTRA_REFRESH:
+			return TEXT("NV_ENC_PIC_TYPE_INTRA_REFRESH");
+		default:
+			checkNoEntry();
+			return TEXT("Unknown");
+		}
+	}
+
+	static bool D3D_ShouldCreateWithD3DDebug()
+	{
+		// Use a debug device if specified on the command line.
+		static bool bCreateWithD3DDebug =
+			FParse::Param(FCommandLine::Get(), TEXT("d3ddebug")) ||
+			FParse::Param(FCommandLine::Get(), TEXT("d3debug")) ||
+			FParse::Param(FCommandLine::Get(), TEXT("dxdebug"));
+		return bCreateWithD3DDebug;
+	}
+
+	static bool D3D_ShouldAllowAsyncResourceCreation()
+	{
+		static bool bAllowAsyncResourceCreation = !FParse::Param(FCommandLine::Get(), TEXT("nod3dasync"));
+		return bAllowAsyncResourceCreation;
+	}
+
+} // anonymous namespace
+
+// Video encoder implementation based on NVIDIA Video Codecs SDK: https://developer.nvidia.com/nvidia-video-codec-sdk
+// Uses only encoder part
+class FNvVideoEncoder : public FVideoEncoder
+{
+public:
+
+	FNvVideoEncoder();
+	~FNvVideoEncoder();
+
+	//
+	// FVideoEncoder interface
+	//
+	const TCHAR* GetName() const override;
+	const TCHAR* GetType() const override;
+	bool Initialize(const FVideoEncoderConfig& InConfig) override;
+	void Shutdown() override;
+	bool CopyTexture(FTexture2DRHIRef Texture, FTimespan CaptureTs, FTimespan Duration, FBufferId& OutBufferId, FIntPoint Resolution) override;
+	void Drop(FBufferId BufferId) override;
+	void Encode(FBufferId BufferId, bool bForceKeyFrame, uint32 Bitrate, TUniquePtr<FEncoderVideoFrameCookie> Cookie) override;
+	FVideoEncoderConfig GetConfig() const override;
+	bool SetBitrate(uint32 Bitrate) override;
+	bool SetFramerate(uint32 Framerate) override;
+	bool SetParameter(const FString& Parameter, const FString& Value) override;
+
+private:
+
+	struct FInputFrame
+	{
+		FInputFrame() {}
+		UE_NONCOPYABLE(FInputFrame);
+		void* RegisteredResource = nullptr;
+		NV_ENC_INPUT_PTR MappedResource = nullptr;
+		NV_ENC_BUFFER_FORMAT BufferFormat;
+		FTexture2DRHIRef Texture;
+		ID3D11Texture2D* SharedTexture = nullptr;
+		bool bForceKeyFrame = false;
+		FTimespan CaptureTs;
+		FTimespan Duration;
+		FGPUFenceRHIRef CopyFence;
+	};
+
+	struct FOutputFrame
+	{
+		FOutputFrame() {}
+		UE_NONCOPYABLE(FOutputFrame);
+		NV_ENC_OUTPUT_PTR BitstreamBuffer = nullptr;
+		HANDLE EventHandle = nullptr;
+		TUniquePtr<FEncoderVideoFrameCookie> Cookie;
+	};
+
+	enum class EFrameState
+	{
+		Free,
+		Capturing,
+		Captured,
+		Encoding
+	};
+
+	struct FFrame
+	{
+		FFrame() {}
+		UE_NONCOPYABLE(FFrame);
+
+		// Array index of this FFrame. This is set at startup, and should never be changed
+		FBufferId Id = 0;
+
+		TAtomic<EFrameState> State = { EFrameState::Free };
+		// Bitrate requested at the time the video encoder asked us to encode this frame
+		// We save this, because we can't use it at the moment we receive it.
+		uint32 BitrateRequested = 0;
+		FInputFrame InputFrame;
+		FOutputFrame OutputFrame;
+		uint64 FrameIdx = 0;
+
+		// Some timestamps to track how long a frame spends in each step
+		FTimespan CopyBufferStartTs;
+		FTimespan CopyBufferFinishTs;
+		FTimespan EncodingStartTs;
+		FTimespan EncodingFinishTs;
+	};
+
+	bool InitFrameInputBuffer(FFrame& Frame, uint32 Width, uint32 Heigh);
+	bool InitializeResources();
+	void ReleaseFrameInputBuffer(FFrame& Frame);
+	void ReleaseResources();
+	void RegisterAsyncEvent(void** OutEvent);
+	void UnregisterAsyncEvent(void* Event);
+	FFrame* CheckForFinishedCopy();
+
+	bool UpdateFramerate();
+
+	/**
+	*  Update some encoder settings
+	*
+	* @param ResolutionIf both X and Y are NOT 0, then set the encoder resolution
+	* @param Bitrate If not 0, then set the encoder average bitrate
+	*/
+	void UpdateNvEncConfig(FIntPoint Resolution, uint32 Bitrate);
+
+	void UpdateRes(FFrame& Frame, FIntPoint Resolution);
+	void CopyTexture(const FTexture2DRHIRef& Texture, FFrame& Frame, FIntPoint Resolution);
+
+	void EncoderCheckLoop();
+
+	void ProcessFrame(FFrame& Frame);
+
+	void UpdateSettings(FInputFrame& InputFrame, uint32 Bitrate);
+	void SubmitFrameToEncoder(FFrame& Frame);
+
+	bool bInitialized = false;
+	void* DllHandle = nullptr;
+	TUniquePtr<NV_ENCODE_API_FUNCTION_LIST> NvEncodeAPI;
+	void* EncoderInterface = nullptr;
+	NV_ENC_INITIALIZE_PARAMS NvEncInitializeParams;
+
+	// Used to atomically change NvEnc settings, so if the outside calls GetConfig, it
+	// gets a valid result, instead of something that was in the middle of being updated.
+	mutable FCriticalSection ConfigCS;
+
+	NV_ENC_CONFIG NvEncConfig;
+	uint32 CapturedFrameCount = 0; // of captured, not encoded frames
+	static constexpr uint32 NumBufferedFrames = 3;
+	FFrame BufferedFrames[NumBufferedFrames];
+	TUniquePtr<FThread> EncoderThread;
+	FThreadSafeBool bExitEncoderThread = false;
+
+	// Desired config.
+	// This is not applied immediately. It's applied when the next frame is sent to
+	// the encoder
+	FVideoEncoderConfig Config;
+	FH264Settings ConfigH264;
+
+	class FEncoderDevice
+	{
+	public:
+		FEncoderDevice();
+
+		TRefCountPtr<ID3D11Device> Device;
+		TRefCountPtr<ID3D11DeviceContext> DeviceContext;
+	};
+
+	TUniquePtr<FEncoderDevice> EncoderDevice;
+
+#if NVENC_VIDEO_ENCODER_DEBUG
+	// This is just for debugging
+	void ClearFrame(FFrame& Frame);
+	// Timings in milliseconds. Just for debugging
+	struct FFrameTiming
+	{
+		// 0 : CopyBufferStart -> CopyBufferFinish
+		// 1 : CopyBufferStart -> EncodingStart
+		// 2 : CopyBufferStart -> EncodingFinish
+		double Total[3];
+		// 0 : CopyBufferStart -> CopyBufferFinish
+		// 1 : CopyBufferFinish -> EncodingStart
+		// 2 : EncodingStart -> EncodingFinish
+		double Steps[3];
+	};
+	TArray<FFrameTiming> Timings;
+#endif
+
+	// When we receive an "Encode" call, we can't send to the encoder right away because maybe the
+	// texture copy hasn't completed yet.
+	// So, we put them in this queue and keep checking if the texture copy finished
+	TQueue<FFrame*> CopyingQueue;
+};
+
+//////////////////////////////////////////////////////////////////////////
+// FNvVideoEncoder implementation
+//////////////////////////////////////////////////////////////////////////
+
+FNvVideoEncoder::FEncoderDevice::FEncoderDevice()
+{
+	if (GDynamicRHI)
+	{
+		FString RHIName = GDynamicRHI->GetName();
+
+		TRefCountPtr<IDXGIDevice> DXGIDevice;
+		uint32 DeviceFlags = D3D_ShouldAllowAsyncResourceCreation() ? 0 : D3D11_CREATE_DEVICE_SINGLETHREADED;
+		if (D3D_ShouldCreateWithD3DDebug())
+			DeviceFlags |= D3D11_CREATE_DEVICE_DEBUG;
+		D3D_FEATURE_LEVEL FeatureLevel = D3D_FEATURE_LEVEL_11_1;
+		TRefCountPtr<IDXGIAdapter> Adapter;
+
+		if (RHIName == TEXT("D3D11"))
+		{
+			auto UE4D3DDevice = static_cast<ID3D11Device*>(GDynamicRHI->RHIGetNativeDevice());
+			checkf(UE4D3DDevice != nullptr, TEXT("Cannot initialize NvEnc with invalid device"));
+			CHECK_HR_VOID(UE4D3DDevice->QueryInterface(__uuidof(IDXGIDevice), (void**)DXGIDevice.GetInitReference()));
+			CHECK_HR_VOID(DXGIDevice->GetAdapter(Adapter.GetInitReference()));
+			FeatureLevel = D3D_FEATURE_LEVEL_11_0;
+		}
+		else if (RHIName == TEXT("D3D12"))
+		{
+			auto UE4D3DDevice = static_cast<ID3D12Device*>(GDynamicRHI->RHIGetNativeDevice());
+			checkf(UE4D3DDevice != nullptr, TEXT("Cannot initialize NvEnc with invalid device"));
+			LUID AdapterLuid = UE4D3DDevice->GetAdapterLuid();
+			TRefCountPtr<IDXGIFactory4> DXGIFactory;
+			CHECK_HR_VOID(CreateDXGIFactory(IID_PPV_ARGS(DXGIFactory.GetInitReference())));
+			// To use a shared texture from D3D12, we need to use a D3D 11.1 device, because we need the
+			// D3D11Device1::OpenSharedResource1 method
+			FeatureLevel = D3D_FEATURE_LEVEL_11_1;
+			CHECK_HR_VOID(DXGIFactory->EnumAdapterByLuid(AdapterLuid, IID_PPV_ARGS(Adapter.GetInitReference())));
+		}
+		else
+		{
+			UE_LOG(LogAVEncoder, Fatal, TEXT("NvEnc requires D3D11/D3D12"));
+			return;
+		}
+
+		D3D_FEATURE_LEVEL ActualFeatureLevel;
+
+		CHECK_HR_VOID(D3D11CreateDevice(
+			Adapter,
+			D3D_DRIVER_TYPE_UNKNOWN,
+			NULL,
+			DeviceFlags,
+			&FeatureLevel,
+			1,
+			D3D11_SDK_VERSION,
+			Device.GetInitReference(),
+			&ActualFeatureLevel,
+			DeviceContext.GetInitReference()
+		));
+
+		// If we are using D3D12, make sure we got a 11.1 device
+		if (FeatureLevel == D3D_FEATURE_LEVEL_11_1 && ActualFeatureLevel!=D3D_FEATURE_LEVEL_11_1)
+		{
+			UE_LOG(LogAVEncoder, Fatal, TEXT("Failed to create a D3D 11.1 device. This is needed when using the D3D12 renderer."));
+		}
+	}
+	else
+	{
+		UE_LOG(LogAVEncoder, Error, TEXT("Attempting to create Encoder Device without existing RHI"));
+	}
+}
+
+const uint32 BitstreamSize = 1024 * 1024 * 2;
+
+#define NV_RESULT(NvFunction) NvFunction == NV_ENC_SUCCESS
+
+FNvVideoEncoder::FNvVideoEncoder()
+{
+	DllHandle = FPlatformProcess::GetDllHandle(GetDllName());
+	checkf(DllHandle != nullptr, TEXT("Failed to load NvEncode dll"));
+	if (!DllHandle)
+	{
+		return;
+	}
+
+}
+
+FNvVideoEncoder::~FNvVideoEncoder()
+{
+	if (DllHandle)
+	{
+		UE_LOG(LogAVEncoder, Fatal, TEXT("FNvVideoEncoder Shutdown not called before destruction."));
+	}
+}
+
+const TCHAR* FNvVideoEncoder::GetName() const
+{
+	return TEXT("h264.nvenc");
+}
+const TCHAR* FNvVideoEncoder::GetType() const
+{
+	return TEXT("h264");
+}
+
+bool FNvVideoEncoder::Initialize(const FVideoEncoderConfig& InConfig)
+{
+	check(!bInitialized);
+
+	Config = InConfig;
+	ConfigH264 = {};
+	ReadH264Settings(Config.Options, ConfigH264);
+
+	UE_LOG(LogAVEncoder, Log, TEXT("FNvVideoEncoder initialization with %u*%d, %u FPS"),
+		Config.Width, Config.Height, Config.Framerate);
+
+	EncoderDevice = MakeUnique<FEncoderDevice>();
+
+	_NVENCSTATUS Result;
+
+	// Load NvEnc dll and create an NvEncode API instance
+	{
+		// define a function pointer for creating an instance of nvEncodeAPI
+		typedef NVENCSTATUS(NVENCAPI *NVENCAPIPROC)(NV_ENCODE_API_FUNCTION_LIST*);
+		NVENCAPIPROC NvEncodeAPICreateInstanceFunc;
+
+#if defined PLATFORM_WINDOWS
+#	pragma warning(push)
+#		pragma warning(disable: 4191) // https://stackoverflow.com/a/4215425/453271
+		NvEncodeAPICreateInstanceFunc = (NVENCAPIPROC)GetProcAddress((HMODULE)DllHandle, "NvEncodeAPICreateInstance");
+#	pragma warning(pop)
+#else
+		NvEncodeAPICreateInstanceFunc = (NVENCAPIPROC)dlsym(DllHandle, "NvEncodeAPICreateInstance");
+#endif
+		checkf(NvEncodeAPICreateInstanceFunc != nullptr, TEXT("NvEncodeAPICreateInstance failed"));
+		NvEncodeAPI.Reset(new NV_ENCODE_API_FUNCTION_LIST);
+		FMemory::Memzero(NvEncodeAPI.Get(), sizeof(NV_ENCODE_API_FUNCTION_LIST));
+		NvEncodeAPI->version = NV_ENCODE_API_FUNCTION_LIST_VER;
+		Result = NvEncodeAPICreateInstanceFunc(NvEncodeAPI.Get());
+		checkf(NV_RESULT(Result), TEXT("Unable to create NvEnc API function list: error %d"), Result);
+	}
+	// Open an encoding session
+	{
+		NV_ENC_OPEN_ENCODE_SESSION_EX_PARAMS OpenEncodeSessionExParams;
+		FMemory::Memzero(OpenEncodeSessionExParams);
+		OpenEncodeSessionExParams.version = NV_ENC_OPEN_ENCODE_SESSION_EX_PARAMS_VER;
+		OpenEncodeSessionExParams.device = EncoderDevice->Device;
+		OpenEncodeSessionExParams.deviceType = NV_ENC_DEVICE_TYPE_DIRECTX;	// Currently only DX11 is supported
+		OpenEncodeSessionExParams.apiVersion = NVENCAPI_VERSION;
+		Result = NvEncodeAPI->nvEncOpenEncodeSessionEx(&OpenEncodeSessionExParams, &EncoderInterface);
+		checkf(NV_RESULT(Result), TEXT("Unable to open NvEnc encoding session (status: %d)"), Result);
+	}
+	// Set initialization parameters
+	{
+		FMemory::Memzero(NvEncInitializeParams);
+		NvEncInitializeParams.version = NV_ENC_INITIALIZE_PARAMS_VER;
+		// hardcoded to FullHD for now, if actual resolution is different it will be changed dynamically
+		NvEncInitializeParams.encodeWidth = NvEncInitializeParams.darWidth = Config.Width;
+		NvEncInitializeParams.encodeHeight = NvEncInitializeParams.darHeight = Config.Height;
+		NvEncInitializeParams.encodeGUID = NV_ENC_CODEC_H264_GUID;
+
+		if (Config.Preset == FVideoEncoderConfig::EPreset::LowLatency)
+			NvEncInitializeParams.presetGUID = NV_ENC_PRESET_LOW_LATENCY_HQ_GUID;
+		else if (Config.Preset == FVideoEncoderConfig::EPreset::HighQuality)
+			NvEncInitializeParams.presetGUID = NV_ENC_PRESET_HQ_GUID;
+		else
+		{
+			check(false);
+		}
+
+		NvEncInitializeParams.frameRateNum = Config.Framerate;
+		NvEncInitializeParams.frameRateDen = 1;
+		NvEncInitializeParams.enablePTD = 1;
+		NvEncInitializeParams.reportSliceOffsets = 0;
+		NvEncInitializeParams.enableSubFrameWrite = 0;
+		NvEncInitializeParams.encodeConfig = &NvEncConfig;
+		NvEncInitializeParams.maxEncodeWidth = 3840;
+		NvEncInitializeParams.maxEncodeHeight = 2160;
+		FParse::Value(FCommandLine::Get(), TEXT("NvEncMaxEncodeWidth="), NvEncInitializeParams.maxEncodeWidth);
+		FParse::Value(FCommandLine::Get(), TEXT("NvEncMaxEncodeHeight="), NvEncInitializeParams.maxEncodeHeight);
+	}
+	// Get preset config and tweak it accordingly
+	{
+		NV_ENC_PRESET_CONFIG PresetConfig;
+		FMemory::Memzero(PresetConfig);
+		PresetConfig.version = NV_ENC_PRESET_CONFIG_VER;
+		PresetConfig.presetCfg.version = NV_ENC_CONFIG_VER;
+		Result = NvEncodeAPI->nvEncGetEncodePresetConfig(EncoderInterface, NvEncInitializeParams.encodeGUID, NvEncInitializeParams.presetGUID, &PresetConfig);
+		checkf(NV_RESULT(Result), TEXT("Failed to select NVEncoder preset config (status: %d)"), Result);
+		FMemory::Memcpy(&NvEncConfig, &PresetConfig.presetCfg, sizeof(NV_ENC_CONFIG));
+
+		NvEncConfig.profileGUID = Config.Preset == FVideoEncoderConfig::EPreset::LowLatency ? NV_ENC_H264_PROFILE_BASELINE_GUID : NV_ENC_H264_PROFILE_MAIN_GUID;
+
+		NvEncConfig.gopLength = NvEncInitializeParams.frameRateNum; // once a sec
+
+		NV_ENC_RC_PARAMS& RcParams = NvEncConfig.rcParams;
+		RcParams.rateControlMode = ToNvEncRcMode(ConfigH264.RcMode);
+
+		RcParams.enableMinQP = true;
+		RcParams.minQP = { 20, 20, 20 };
+
+		RcParams.maxBitRate = Config.MaxBitrate;
+		RcParams.averageBitRate = FMath::Min(Config.Bitrate, RcParams.maxBitRate);
+
+		NvEncConfig.encodeCodecConfig.h264Config.idrPeriod = NvEncConfig.gopLength;
+
+		// configure "entire frame as a single slice"
+		// seems WebRTC implementation doesn't work well with slicing, default mode 
+		// (Mode=3/ModeData=4 - 4 slices per frame) produces (rarely) grey full screen or just top half of it. 
+		// it also can be related with our handling of slices in proxy's FakeVideoEncoder
+		if (Config.Preset == FVideoEncoderConfig::EPreset::LowLatency)
+		{
+			NvEncConfig.encodeCodecConfig.h264Config.sliceMode = 0;
+			NvEncConfig.encodeCodecConfig.h264Config.sliceModeData = 0;
+		}
+		else
+		{
+			NvEncConfig.encodeCodecConfig.h264Config.sliceMode = 3;
+			NvEncConfig.encodeCodecConfig.h264Config.sliceModeData = 1;
+		}
+
+		// let encoder slice encoded frame so they can fit into RTP packets
+		// commented out because at some point it started to produce immediately visible visual artifacts on players
+		//NvEncConfig.encodeCodecConfig.h264Config.sliceMode = 1;
+		//NvEncConfig.encodeCodecConfig.h264Config.sliceModeData = 1100; // max bytes per slice
+
+		// repeat SPS/PPS with each key-frame for a case when the first frame (with mandatory SPS/PPS) 
+		// was dropped by WebRTC
+		NvEncConfig.encodeCodecConfig.h264Config.repeatSPSPPS = 1;
+
+		// maybe doesn't have an effect, high level is chosen because we aim at high bitrate
+		NvEncConfig.encodeCodecConfig.h264Config.level = Config.Preset==FVideoEncoderConfig::EPreset::LowLatency ? NV_ENC_LEVEL_H264_52 : NV_ENC_LEVEL_H264_51;
+
+	}
+
+	// Get encoder capability
+	{
+		NV_ENC_CAPS_PARAM CapsParam;
+		FMemory::Memzero(CapsParam);
+		CapsParam.version = NV_ENC_CAPS_PARAM_VER;
+		CapsParam.capsToQuery = NV_ENC_CAPS_ASYNC_ENCODE_SUPPORT;
+		int32 AsyncMode = 0;
+		Result = NvEncodeAPI->nvEncGetEncodeCaps(EncoderInterface, NvEncInitializeParams.encodeGUID, &CapsParam, &AsyncMode);
+		checkf(NV_RESULT(Result), TEXT("Failed to get NVEncoder capability params (status: %d)"), Result);
+		if (AsyncMode == 0)
+		{
+			UE_LOG(LogAVEncoder, Fatal, TEXT("NvEnc doesn't support async mode"));
+			return false;
+		}
+
+		NvEncInitializeParams.enableEncodeAsync = true;
+	}
+
+	Result = NvEncodeAPI->nvEncInitializeEncoder(EncoderInterface, &NvEncInitializeParams);
+	checkf(NV_RESULT(Result), TEXT("Failed to initialize NVEncoder (status: %d)"), Result);
+
+	FBufferId Id = 0;
+	for(FFrame& Frame : BufferedFrames)
+	{
+		Frame.Id = Id++;
+	}
+
+	if (!InitializeResources())
+	{
+		return false;
+	}
+
+	EncoderThread = MakeUnique<FThread>(TEXT("NvVideoEncoder"), [this]() { EncoderCheckLoop(); });
+
+	UE_LOG(LogAVEncoder, Log, TEXT("NvEnc initialised"));
+	
+	bInitialized = true;
+	return true;
+}
+
+void FNvVideoEncoder::Shutdown()
+{
+	if (!DllHandle)
+	{
+		return;
+	}
+
+	if (EncoderThread)
+	{
+		bExitEncoderThread = true;
+
+		// Exit encoder runnable thread before shutting down NvEnc interface
+		EncoderThread->Join();
+	}
+
+	ReleaseResources();
+
+	if (EncoderInterface)
+	{
+		_NVENCSTATUS Result = NvEncodeAPI->nvEncDestroyEncoder(EncoderInterface);
+		checkf(NV_RESULT(Result), TEXT("Failed to destroy NvEnc interface (status: %d)"), Result);
+		EncoderInterface = nullptr;
+	}
+
+#if defined PLATFORM_WINDOWS
+		FPlatformProcess::FreeDllHandle(DllHandle);
+#else
+		dlclose(DllHandle);
+#endif
+		DllHandle = nullptr;
+}
+
+FVideoEncoderConfig FNvVideoEncoder::GetConfig() const
+{
+	FVideoEncoderConfig Cfg;
+	FScopeLock ScopedLock(&ConfigCS);
+	Cfg.Bitrate = NvEncConfig.rcParams.averageBitRate;
+	Cfg.Framerate = NvEncInitializeParams.frameRateNum;
+	Cfg.Width = NvEncInitializeParams.encodeWidth;
+	Cfg.Height = NvEncInitializeParams.encodeHeight;
+	return Cfg;
+}
+
+bool FNvVideoEncoder::SetBitrate(uint32 Bitrate)
+{
+	Config.Bitrate;
+	return true;
+}
+
+bool FNvVideoEncoder::SetFramerate(uint32 Framerate)
+{
+	Config.Framerate = Framerate;
+	return true;
+}
+
+bool FNvVideoEncoder::SetParameter(const FString& Parameter, const FString& Value)
+{
+	return ReadH264Setting(Parameter, Value, ConfigH264);
+}
+
+bool FNvVideoEncoder::CopyTexture(const FTexture2DRHIRef Texture, FTimespan CaptureTs, FTimespan Duration, FBufferId& OutBufferId, FIntPoint Resolution)
+{
+	check(IsInRenderingThread());
+
+	// Find a free slot we can use
+	FFrame* Frame = nullptr;
+	for (FFrame& Slot : BufferedFrames)
+	{
+		if (Slot.State.Load() == EFrameState::Free)
+		{
+			Frame = &Slot;
+			OutBufferId = Slot.Id;
+			break;
+		}
+	}
+
+	if (!Frame)
+	{
+		UE_LOG(LogAVEncoder, Verbose, TEXT("Frame dropped because NvEnc queue is full"));
+		return false;
+	}
+
+	Frame->FrameIdx = CapturedFrameCount++;
+	Frame->InputFrame.CaptureTs = CaptureTs;
+	Frame->InputFrame.Duration = Duration;
+	Frame->CopyBufferStartTs = FTimespan::FromSeconds(FPlatformTime::Seconds());
+
+#if NVENC_VIDEO_ENCODER_DEBUG
+	// By clearing the frame at this point, we can catch the occasional glimpse of a solid color
+	// frame in PixelStreaming if there are any bugs detecting when the copy finished
+	ClearFrame(*Frame);
+#endif
+
+	CopyTexture(Texture, *Frame, Resolution);
+
+	UE_LOG(LogAVEncoder, Verbose, TEXT("Buffer #%d (%d) captured"), Frame->FrameIdx, OutBufferId);
+	Frame->State = EFrameState::Capturing;
+
+	return true;
+}
+
+void FNvVideoEncoder::CopyTexture(const FTexture2DRHIRef& Texture, FFrame& Frame, FIntPoint Resolution)
+{
+	SCOPE_CYCLE_COUNTER(STAT_NvEnc_CopyTexture);
+	UpdateRes(Frame, Resolution.Size() ? Resolution : Texture->GetSizeXY());
+	CopyTextureImpl(Texture, Frame.InputFrame.Texture, Frame.InputFrame.CopyFence);
+}
+
+FNvVideoEncoder::FFrame* FNvVideoEncoder::CheckForFinishedCopy()
+{
+	FFrame* Frame;
+	if (!CopyingQueue.Peek(Frame))
+	{
+		return nullptr;
+	}
+
+	{
+		EFrameState State = Frame->State.Load();
+		checkf(State == EFrameState::Capturing, TEXT("Buffer %d : Expected state %d, but found %d"), Frame->Id, (int)EFrameState::Captured, (int)State);
+	}
+
+	if (Frame->InputFrame.CopyFence->Poll())
+	{
+		CopyingQueue.Pop();
+		Frame->State = EFrameState::Captured;
+		Frame->CopyBufferFinishTs = FTimespan::FromSeconds(FPlatformTime::Seconds());
+		return Frame;
+	}
+	else
+	{
+		return nullptr;
+	}
+
+}
+
+void FNvVideoEncoder::Encode(FBufferId BufferId, bool bForceKeyFrame, uint32 Bitrate, TUniquePtr<FEncoderVideoFrameCookie> Cookie)
+{
+	FFrame& Frame = BufferedFrames[BufferId];
+
+	{
+		EFrameState State = Frame.State.Load();
+		checkf(State == EFrameState::Capturing, TEXT("Buffer %d : Expected state %d, but found %d"), BufferId, (int)EFrameState::Captured, (int)State);
+	}
+
+	Frame.InputFrame.bForceKeyFrame = bForceKeyFrame;
+	Frame.BitrateRequested = Bitrate;
+	Frame.OutputFrame.Cookie = MoveTemp(Cookie);
+	CopyingQueue.Enqueue(&Frame);
+}
+
+void FNvVideoEncoder::Drop(FBufferId BufferId)
+{
+	FFrame& Frame = BufferedFrames[BufferId];
+
+	{
+		EFrameState State = Frame.State.Load();
+		checkf(State == EFrameState::Capturing, TEXT("Buffer %d: Expected state %d, found %d")
+			, BufferId, (int)EFrameState::Capturing, (int)State);
+	}
+
+	Frame.State = EFrameState::Free;
+
+	UE_LOG(LogAVEncoder, Log, TEXT("Buffer #%d (%d) dropped"), BufferedFrames[BufferId].FrameIdx, BufferId);
+}
+
+void FNvVideoEncoder::UpdateNvEncConfig(FIntPoint Resolution, uint32 Bitrate)
+{
+	bool bSettingsChanged = false;
+	bool bResolutionChanged = false;
+
+	// Putting all this in a single scope, so we can put changes to the config struct
+	// inside a big lock
+	{
+		FScopeLock ScopedLock(&ConfigCS);
+
+		//
+		// If an explicit Bitrate was specified, use that one, if not, use the one
+		// from the Config struct
+		//
+		if (Bitrate)
+		{
+			if (NvEncConfig.rcParams.averageBitRate != Bitrate)
+			{
+				NvEncConfig.rcParams.averageBitRate = Bitrate;
+				Config.Bitrate = Bitrate;
+				bSettingsChanged = true;
+			}
+		}
+		else 
+		{ 
+			if (NvEncConfig.rcParams.averageBitRate != Config.Bitrate)
+			{
+				NvEncConfig.rcParams.averageBitRate = Config.Bitrate;
+				bSettingsChanged = true;
+			}
+		}
+
+		if (NvEncConfig.rcParams.minQP.qpIntra != ConfigH264.QP)
+		{
+			NvEncConfig.rcParams.minQP.qpIntra = NvEncConfig.rcParams.minQP.qpInterP = NvEncConfig.rcParams.minQP.qpInterB = ConfigH264.QP;
+			UE_LOG(LogAVEncoder, Log, TEXT("MinQP %u"), ConfigH264.QP);
+			bSettingsChanged = true;
+		}
+
+		NV_ENC_PARAMS_RC_MODE RcMode = ToNvEncRcMode(ConfigH264.RcMode);
+		if (RcMode != NvEncConfig.rcParams.rateControlMode)
+		{
+			NvEncConfig.rcParams.rateControlMode = RcMode;
+			UE_LOG(LogAVEncoder, Log, TEXT("Rate Control mode %s"), ToString(RcMode));
+			bSettingsChanged = true;
+		}
+
+		if (UpdateFramerate())
+		{
+			bSettingsChanged = true;
+		}
+
+		// Only try and change resolution if required
+		if (Resolution.X && Resolution.Y)
+		{
+			if (Resolution.X != NvEncInitializeParams.encodeWidth ||
+				Resolution.Y != NvEncInitializeParams.encodeHeight)
+			{
+				NvEncInitializeParams.encodeWidth = NvEncInitializeParams.darWidth = Resolution.X;
+				NvEncInitializeParams.encodeHeight = NvEncInitializeParams.darHeight = Resolution.Y;
+
+				bSettingsChanged = true;
+				bResolutionChanged = true;
+			}
+		}
+
+	}
+
+
+	if (bSettingsChanged)
+	{
+		NV_ENC_RECONFIGURE_PARAMS NvEncReconfigureParams;
+		FMemory::Memcpy(&NvEncReconfigureParams.reInitEncodeParams, &NvEncInitializeParams, sizeof(NvEncInitializeParams));
+		NvEncReconfigureParams.version = NV_ENC_RECONFIGURE_PARAMS_VER;
+		NvEncReconfigureParams.forceIDR = bResolutionChanged;
+
+		_NVENCSTATUS Result = NvEncodeAPI->nvEncReconfigureEncoder(EncoderInterface, &NvEncReconfigureParams);
+		checkf(NV_RESULT(Result), TEXT("Failed to reconfigure encoder (status: %d)"), Result);
+	}
+}
+
+bool FNvVideoEncoder::UpdateFramerate()
+{
+
+	if (NvEncInitializeParams.frameRateNum != Config.Framerate)
+	{
+		NvEncInitializeParams.frameRateNum = Config.Framerate;
+		UE_LOG(LogAVEncoder, Log, TEXT("NvEnc reconfigured to %d FPS"), NvEncInitializeParams.frameRateNum);
+		return true;
+	}
+
+	return false;
+}
+
+// checks if resolution changed, either the game res changed or new streaming resolution was specified by the console var
+void FNvVideoEncoder::UpdateRes(FFrame& Frame, FIntPoint Resolution)
+{
+	check(IsInRenderingThread());
+
+	FInputFrame& InputFrame = Frame.InputFrame;
+
+	// check if target resolution matches our currently allocated `InputFrame.Texture` resolution
+	if (InputFrame.Texture->GetSizeX() == Resolution.X && InputFrame.Texture->GetSizeY() == Resolution.Y)
+	{
+		return;
+	}
+
+	// reallocate and re-register InputFrame with NvEnc
+	ReleaseFrameInputBuffer(Frame);
+	verify(InitFrameInputBuffer(Frame, Resolution.X, Resolution.Y));
+}
+
+void FNvVideoEncoder::SubmitFrameToEncoder(FFrame& Frame)
+{
+	check(Frame.State.Load() == EFrameState::Captured);
+
+	SCOPE_CYCLE_COUNTER(STAT_NvEnc_SubmitFrameToEncoder);
+
+	Frame.State = EFrameState::Encoding;
+	Frame.EncodingStartTs = FTimespan::FromSeconds(FPlatformTime::Seconds());
+
+	NV_ENC_PIC_PARAMS PicParams;
+	FMemory::Memzero(PicParams);
+	PicParams.version = NV_ENC_PIC_PARAMS_VER;
+	PicParams.inputBuffer = Frame.InputFrame.MappedResource;
+	PicParams.bufferFmt = Frame.InputFrame.BufferFormat;
+	PicParams.inputWidth = NvEncInitializeParams.encodeWidth;
+	PicParams.inputHeight = NvEncInitializeParams.encodeHeight;
+	PicParams.outputBitstream = Frame.OutputFrame.BitstreamBuffer;
+	PicParams.completionEvent = Frame.OutputFrame.EventHandle;
+	PicParams.inputTimeStamp = Frame.FrameIdx;
+	PicParams.pictureStruct = NV_ENC_PIC_STRUCT_FRAME;
+
+	if (Frame.InputFrame.bForceKeyFrame)
+		PicParams.encodePicFlags |= NV_ENC_PIC_FLAG_FORCEIDR;
+
+	_NVENCSTATUS Result = NvEncodeAPI->nvEncEncodePicture(EncoderInterface, &PicParams);
+	checkf(NV_RESULT(Result), TEXT("Failed to encode frame (status: %d)"), Result);
+}
+
+void FNvVideoEncoder::EncoderCheckLoop()
+{
+	// This thread will both encode frames and will also wait for the next frame
+	// to finish encoding.
+	TQueue<FFrame*> CurrentlyEncodingQueue;
+
+	while (!bExitEncoderThread)
+	{
+		//
+		// Check if any frames finished copying so we can submit then to the encoder
+		//
+		while(true)
+		{
+			FFrame* Frame = CheckForFinishedCopy();
+			if (!Frame)
+			{
+				break;
+			}
+			UpdateNvEncConfig(Frame->InputFrame.Texture->GetSizeXY(), Frame->BitrateRequested);
+			SubmitFrameToEncoder(*Frame);
+			CurrentlyEncodingQueue.Enqueue(Frame);
+		}
+
+		//
+		// Check for finished encoding work
+		//
+		if (!CurrentlyEncodingQueue.IsEmpty())
+		{
+			HANDLE Handle = (**CurrentlyEncodingQueue.Peek()).OutputFrame.EventHandle;
+			DWORD Result = WaitForSingleObject(Handle, 2);
+			if (Result == WAIT_OBJECT_0)
+			{
+				FFrame* Frame = nullptr;
+				verify(CurrentlyEncodingQueue.Dequeue(Frame));
+				ResetEvent(Frame->OutputFrame.EventHandle);
+				UE_LOG(LogAVEncoder, Verbose, TEXT("Buffer #%d (%d) encoded"), Frame->FrameIdx, Frame->Id);
+				ProcessFrame(*Frame);
+			}
+			else if (Result == WAIT_TIMEOUT)
+			{
+				// Nothing to do. This is expected
+			}
+			else
+			{
+				check(0 && "Unexpected code path");
+			}
+		}
+	}
+}
+
+void FNvVideoEncoder::ProcessFrame(FFrame& Frame)
+{
+	check(Frame.State.Load() == EFrameState::Encoding);
+
+	FOutputFrame& OutputFrame = Frame.OutputFrame;
+
+	FAVPacket Packet(EPacketType::Video);
+	NV_ENC_PIC_TYPE PicType;
+	// Retrieve encoded frame from output buffer
+	{
+		SCOPE_CYCLE_COUNTER(STAT_NvEnc_RetrieveEncodedFrame);
+
+		NV_ENC_LOCK_BITSTREAM LockBitstream;
+		FMemory::Memzero(LockBitstream);
+		LockBitstream.version = NV_ENC_LOCK_BITSTREAM_VER;
+		LockBitstream.outputBitstream = OutputFrame.BitstreamBuffer;
+		LockBitstream.doNotWait = NvEncInitializeParams.enableEncodeAsync;
+
+		_NVENCSTATUS Result = NvEncodeAPI->nvEncLockBitstream(EncoderInterface, &LockBitstream);
+		checkf(NV_RESULT(Result), TEXT("Failed to lock bitstream (status: %d)"), Result);
+
+		PicType = LockBitstream.pictureType;
+		checkf(PicType == NV_ENC_PIC_TYPE_IDR || Frame.InputFrame.bForceKeyFrame==false, TEXT("key frame requested by but not provided by NvEnc. NvEnc provided %d"), (int)PicType);
+		Packet.Video.bKeyFrame = (PicType == NV_ENC_PIC_TYPE_IDR) ? true : false;
+		Packet.Video.FrameAvgQP = LockBitstream.frameAvgQP; 
+		Packet.Data = TArray<uint8>(reinterpret_cast<const uint8*>(LockBitstream.bitstreamBufferPtr), LockBitstream.bitstreamSizeInBytes);
+		Result = NvEncodeAPI->nvEncUnlockBitstream(EncoderInterface, Frame.OutputFrame.BitstreamBuffer);
+		checkf(NV_RESULT(Result), TEXT("Failed to unlock bitstream (status: %d)"), Result);
+	}
+
+	Frame.EncodingFinishTs = FTimespan::FromSeconds(FPlatformTime::Seconds());
+
+	Packet.Timestamp = Frame.InputFrame.CaptureTs;
+	Packet.Duration = Frame.InputFrame.Duration;
+	Packet.Video.Width = Frame.InputFrame.Texture->GetSizeX();
+	Packet.Video.Height = Frame.InputFrame.Texture->GetSizeY();
+	Packet.Video.Framerate = NvEncInitializeParams.frameRateNum;
+	Packet.Timings.EncodeStartTs = Frame.EncodingStartTs;
+	Packet.Timings.EncodeFinishTs = Frame.EncodingFinishTs;
+
+#if NVENC_VIDEO_ENCODER_DEBUG
+	{
+		FFrameTiming Timing;
+		Timing.Total[0] = (Frame.CopyBufferFinishTs - Frame.CopyBufferStartTs).GetTotalMilliseconds();
+		Timing.Total[1] = (Frame.EncodingStartTs - Frame.CopyBufferStartTs).GetTotalMilliseconds();
+		Timing.Total[2] = (Frame.EncodingFinishTs - Frame.CopyBufferStartTs).GetTotalMilliseconds();
+
+		Timing.Steps[0] = (Frame.CopyBufferFinishTs - Frame.CopyBufferStartTs).GetTotalMilliseconds();
+		Timing.Steps[1] = (Frame.EncodingStartTs - Frame.CopyBufferFinishTs).GetTotalMilliseconds();
+		Timing.Steps[2] = (Frame.EncodingFinishTs - Frame.EncodingStartTs).GetTotalMilliseconds();
+		Timings.Add(Timing);
+		// Limit the array size
+		if (Timings.Num()>1000)
+		{
+			Timings.RemoveAt(0);
+		}
+	}
+#endif
+
+	UE_LOG(LogAVEncoder, VeryVerbose, TEXT("encoded %s ts %lld, %d bytes")
+		, ToString(PicType)
+		, Packet.Timestamp.GetTicks()
+		, (int)Packet.Data.Num());
+
+	{
+		SCOPE_CYCLE_COUNTER(STAT_NvEnc_OnEncodedVideoFrameCallback);
+		OnEncodedVideoFrame(Packet, MoveTemp(Frame.OutputFrame.Cookie));
+	}
+
+	Frame.State = EFrameState::Free;
+}
+
+bool FNvVideoEncoder::InitFrameInputBuffer(FFrame& Frame, uint32 Width, uint32 Height)
+{
+	FInputFrame& InputFrame = Frame.InputFrame;
+
+	// Create (if necessary) and clear the GPU Fence so we can detect when the copy finished
+	if (!InputFrame.CopyFence)
+	{
+		FRHICommandListImmediate& RHICmdList = FRHICommandListExecutor::GetImmediateCommandList();
+		InputFrame.CopyFence = RHICmdList.CreateGPUFence(*FString::Printf(TEXT("PixelStreamingCopy_%d"), Frame.Id));
+	}
+
+	// Create resolved back buffer texture
+	{
+		// Make sure format used here is compatible with NV_ENC_BUFFER_FORMAT specified later in NV_ENC_REGISTER_RESOURCE bufferFormat
+		FRHIResourceCreateInfo CreateInfo;
+		InputFrame.Texture = RHICreateTexture2D(Width, Height, EPixelFormat::PF_B8G8R8A8, 1, 1, TexCreate_RenderTargetable | TexCreate_Shared, CreateInfo);
+	}
+
+	// Share this texture with the encoder device.
+	FString RHIName = GDynamicRHI->GetName();
+
+	if (RHIName == TEXT("D3D11"))
+	{
+		ID3D11Texture2D* ResolvedTexture = (ID3D11Texture2D*)InputFrame.Texture->GetTexture2D()->GetNativeResource();
+
+		TRefCountPtr<IDXGIResource> DXGIResource;
+		CHECK_HR_DEFAULT(ResolvedTexture->QueryInterface(IID_PPV_ARGS(DXGIResource.GetInitReference())));
+
+		//
+		// NOTE : The HANDLE IDXGIResource::GetSharedHandle gives us is NOT an NT Handle, and therefre we should not call CloseHandle on it
+		//
+		HANDLE SharedHandle;
+		CHECK_HR_DEFAULT(DXGIResource->GetSharedHandle(&SharedHandle));
+		CHECK_HR_DEFAULT(EncoderDevice->Device->OpenSharedResource(SharedHandle, __uuidof(ID3D11Texture2D), (LPVOID*)&InputFrame.SharedTexture));
+	}
+	else if (RHIName == TEXT("D3D12"))
+	{
+		auto UE4D3DDevice = static_cast<ID3D12Device*>(GDynamicRHI->RHIGetNativeDevice());
+		static uint32 NamingIdx = 0;
+		ID3D12Resource* ResolvedTexture = (ID3D12Resource*)InputFrame.Texture->GetTexture2D()->GetNativeResource();
+
+		//
+		// NOTE: ID3D12Device::CreateSharedHandle gives as an NT Handle, and so we need to call CloseHandle on it
+		//
+		HANDLE SharedHandle;
+		HRESULT Res1 = UE4D3DDevice->CreateSharedHandle(ResolvedTexture, NULL, GENERIC_ALL, *FString::Printf(TEXT("PixelStreaming_NvEnc_%u"), NamingIdx++), &SharedHandle);
+		CHECK_HR_DEFAULT(Res1);
+
+		TRefCountPtr <ID3D11Device1> Device1;
+		CHECK_HR_DEFAULT(EncoderDevice->Device->QueryInterface(__uuidof(ID3D11Device1), (void**)Device1.GetInitReference()));
+		CHECK_HR_DEFAULT(Device1->OpenSharedResource1(SharedHandle, __uuidof(ID3D11Texture2D), (LPVOID*)&InputFrame.SharedTexture));
+		verify(CloseHandle(SharedHandle));
+	}
+
+	// Register input back buffer
+	{
+		NV_ENC_REGISTER_RESOURCE RegisterResource;
+		FMemory::Memzero(RegisterResource);
+		EPixelFormat PixelFormat = InputFrame.Texture->GetFormat();
+
+		RegisterResource.version = NV_ENC_REGISTER_RESOURCE_VER;
+		RegisterResource.resourceType = NV_ENC_INPUT_RESOURCE_TYPE_DIRECTX;
+		RegisterResource.resourceToRegister = (void*)InputFrame.SharedTexture;
+		RegisterResource.width = Width;
+		RegisterResource.height = Height;
+		RegisterResource.bufferFormat = NV_ENC_BUFFER_FORMAT_ABGR;	// Make sure ResolvedTexture is created with a compatible format
+		_NVENCSTATUS Result = NvEncodeAPI->nvEncRegisterResource(EncoderInterface, &RegisterResource);
+		checkf(NV_RESULT(Result), TEXT("Failed to register input back buffer (status: %d)"), Result);
+
+		InputFrame.RegisteredResource = RegisterResource.registeredResource;
+		InputFrame.BufferFormat = RegisterResource.bufferFormat;
+	}
+	// Map input buffer resource
+	{
+		NV_ENC_MAP_INPUT_RESOURCE MapInputResource;
+		FMemory::Memzero(MapInputResource);
+		MapInputResource.version = NV_ENC_MAP_INPUT_RESOURCE_VER;
+		MapInputResource.registeredResource = InputFrame.RegisteredResource;
+		_NVENCSTATUS Result = NvEncodeAPI->nvEncMapInputResource(EncoderInterface, &MapInputResource);
+		checkf(NV_RESULT(Result), TEXT("Failed to map NvEnc input resource (status: %d)"), Result);
+		InputFrame.MappedResource = MapInputResource.mappedResource;
+	}
+
+	return true;
+}
+
+bool FNvVideoEncoder::InitializeResources()
+{
+	for (uint32 i = 0; i < NumBufferedFrames; ++i)
+	{
+		FFrame& Frame = BufferedFrames[i];
+
+		if (!InitFrameInputBuffer(Frame, NvEncInitializeParams.encodeWidth, NvEncInitializeParams.encodeHeight))
+		{
+			return false;
+		}
+
+		FMemory::Memzero(Frame.OutputFrame);
+		// Create output bitstream buffer
+		{
+			NV_ENC_CREATE_BITSTREAM_BUFFER CreateBitstreamBuffer;
+			FMemory::Memzero(CreateBitstreamBuffer);
+			CreateBitstreamBuffer.version = NV_ENC_CREATE_BITSTREAM_BUFFER_VER;
+			CreateBitstreamBuffer.size = BitstreamSize;
+			CreateBitstreamBuffer.memoryHeap = NV_ENC_MEMORY_HEAP_SYSMEM_CACHED;
+			_NVENCSTATUS Result = NvEncodeAPI->nvEncCreateBitstreamBuffer(EncoderInterface, &CreateBitstreamBuffer);
+			checkf(NV_RESULT(Result), TEXT("Failed to create NvEnc bitstream buffer (status: %d)"), Result);
+			Frame.OutputFrame.BitstreamBuffer = CreateBitstreamBuffer.bitstreamBuffer;
+		}
+
+		RegisterAsyncEvent(&Frame.OutputFrame.EventHandle);
+	}
+
+	return true;
+}
+
+void FNvVideoEncoder::ReleaseFrameInputBuffer(FFrame& Frame)
+{
+	FInputFrame& InputFrame = Frame.InputFrame;
+
+	if (InputFrame.MappedResource)
+	{
+		_NVENCSTATUS Result = NvEncodeAPI->nvEncUnmapInputResource(EncoderInterface, InputFrame.MappedResource);
+		checkf(NV_RESULT(Result), TEXT("Failed to unmap input resource (status: %d)"), Result);
+		InputFrame.MappedResource = nullptr;
+	}
+
+	if (InputFrame.RegisteredResource)
+	{
+		_NVENCSTATUS Result = NvEncodeAPI->nvEncUnregisterResource(EncoderInterface, InputFrame.RegisteredResource);
+		checkf(NV_RESULT(Result), TEXT("Failed to unregister input buffer resource (status: %d)"), Result);
+		InputFrame.RegisteredResource = nullptr;
+	}
+
+	InputFrame.Texture.SafeRelease();
+	if (InputFrame.SharedTexture)
+	{
+		InputFrame.SharedTexture->Release();
+		InputFrame.SharedTexture = nullptr;
+	}
+
+	if (InputFrame.CopyFence)
+	{
+		InputFrame.CopyFence.SafeRelease();
+	}
+}
+
+void FNvVideoEncoder::ReleaseResources()
+{
+	for (uint32 i = 0; i < NumBufferedFrames; ++i)
+	{
+		FFrame& Frame = BufferedFrames[i];
+
+		ReleaseFrameInputBuffer(Frame);
+
+		if (Frame.OutputFrame.BitstreamBuffer)
+		{
+			_NVENCSTATUS Result = NvEncodeAPI->nvEncDestroyBitstreamBuffer(EncoderInterface, Frame.OutputFrame.BitstreamBuffer);
+			checkf(NV_RESULT(Result), TEXT("Failed to destroy output buffer bitstream (status: %d)"), Result);
+			Frame.OutputFrame.BitstreamBuffer = nullptr;
+		}
+
+		if (Frame.OutputFrame.EventHandle)
+		{
+			UnregisterAsyncEvent(Frame.OutputFrame.EventHandle);
+			::CloseHandle(Frame.OutputFrame.EventHandle);
+			Frame.OutputFrame.EventHandle = nullptr;
+		}
+	}
+}
+
+void FNvVideoEncoder::RegisterAsyncEvent(void** OutEvent)
+{
+	NV_ENC_EVENT_PARAMS EventParams;
+	FMemory::Memzero(EventParams);
+	EventParams.version = NV_ENC_EVENT_PARAMS_VER;
+#if defined PLATFORM_WINDOWS
+	EventParams.completionEvent = CreateEvent(nullptr, false, false, nullptr);
+#endif
+	_NVENCSTATUS Result = NvEncodeAPI->nvEncRegisterAsyncEvent(EncoderInterface, &EventParams);
+	checkf(NV_RESULT(Result), TEXT("Failed to register async event (status: %d)"), Result);
+	*OutEvent = EventParams.completionEvent;
+}
+
+void FNvVideoEncoder::UnregisterAsyncEvent(void* Event)
+{
+	if (Event)
+	{
+		NV_ENC_EVENT_PARAMS EventParams;
+		FMemory::Memzero(EventParams);
+		EventParams.version = NV_ENC_EVENT_PARAMS_VER;
+		EventParams.completionEvent = Event;
+		bool Result = NV_RESULT(NvEncodeAPI->nvEncUnregisterAsyncEvent(EncoderInterface, &EventParams));
+		checkf(Result, TEXT("Failed to unregister async event"));
+	}
+}
+
+#if NVENC_VIDEO_ENCODER_DEBUG
+// Fills with a solid colour
+void FNvVideoEncoder::ClearFrame(FFrame& Frame)
+{
+	check(IsInRenderingThread());
+
+	FRHICommandListImmediate& RHICmdList = FRHICommandListExecutor::GetImmediateCommandList();
+
+	static_assert(NumBufferedFrames == 3, "Unexpected number of slots. Please update the array to match.");
+	FLinearColor Colors[NumBufferedFrames] =
+	{
+		FLinearColor(1,0,0),
+		FLinearColor(0,1,0),
+		FLinearColor(0,0,1)
+	};
+
+	FRHIRenderPassInfo RPInfo(Frame.InputFrame.Texture, ERenderTargetActions::Load_Store);
+	TransitionRenderPassTargets(RHICmdList, RPInfo);
+	RHICmdList.BeginRenderPass(RPInfo, TEXT("ClearCanvas"));
+	RHICmdList.SetViewport(0, 0, 0.0f
+		, Frame.InputFrame.Texture->GetSizeXY().X
+		, Frame.InputFrame.Texture->GetSizeXY().Y, 1.0f);
+
+	DrawClearQuad(RHICmdList, Colors[Frame.Id]);
+	RHICmdList.EndRenderPass();
+}
+#endif
+
+
+//////////////////////////////////////////////////////////////////////////
+// FNvVideoEncoderFactory implementation
+//////////////////////////////////////////////////////////////////////////
+
+FNvVideoEncoderFactory::FNvVideoEncoderFactory()
+{
+}
+
+FNvVideoEncoderFactory::~FNvVideoEncoderFactory()
+{
+}
+
+const TCHAR* FNvVideoEncoderFactory::GetName() const
+{
+	return TEXT("nvenc");
+}
+
+TArray<FString> FNvVideoEncoderFactory::GetSupportedCodecs() const
+{
+	TArray<FString> Codecs;
+
+	if (!IsRHIDeviceNVIDIA())
+	{
+		UE_LOG(LogAVEncoder, Log, TEXT("No NvEnc because no NVidia card found"));
+		return Codecs;
+	}
+
+	void* Handle = FPlatformProcess::GetDllHandle(GetDllName());
+	if (Handle == nullptr)
+	{
+		UE_LOG(LogAVEncoder, Error, TEXT("NVidia card found, but no NvEnc DLL installed."));
+		return Codecs;
+	}
+	else
+	{
+		FPlatformProcess::FreeDllHandle(Handle);
+	}
+
+	Codecs.Add(TEXT("h264"));
+	return Codecs;
+}
+
+TUniquePtr<FVideoEncoder> FNvVideoEncoderFactory::CreateEncoder(const FString& Codec)
+{
+	if (Codec=="h264")
+	{
+		return TUniquePtr<FVideoEncoder>(new FNvVideoEncoder());
+	}
+	else
+	{
+		UE_LOG(LogAVEncoder, Error, TEXT("FNvVideoEncoderFactory doesn't support the %s codec"), *Codec);
+		return nullptr;
+	}
+}
+
+
+} // namespace AVEncoder
+
+
diff --git a/Engine/Source/Runtime/AVEncoder/Private/Microsoft/Windows/NvVideoEncoder.h b/Engine/Source/Runtime/AVEncoder/Private/Microsoft/Windows/NvVideoEncoder.h
new file mode 100644
--- /dev/null
+++ b/Engine/Source/Runtime/AVEncoder/Private/Microsoft/Windows/NvVideoEncoder.h
@@ -0,0 +1,26 @@
+// Copyright 1998-2019 Epic Games, Inc. All Rights Reserved.
+
+#pragma once
+
+#include "AVEncoder.h"
+#include "AVEncoderCommon.h"
+
+namespace AVEncoder
+{
+
+/**
+ * Factory for Nvidia's NvEnc
+ */
+class FNvVideoEncoderFactory : public FVideoEncoderFactory
+{
+public:
+	FNvVideoEncoderFactory();
+	~FNvVideoEncoderFactory() override;
+	const TCHAR* GetName() const override;
+	TArray<FString> GetSupportedCodecs() const override;
+	TUniquePtr<FVideoEncoder> CreateEncoder(const FString& Codec) override;
+private:
+};
+
+}
+
diff --git a/Engine/Source/Runtime/AVEncoder/Public/AVEncoder.cpp b/Engine/Source/Runtime/AVEncoder/Public/AVEncoder.cpp
new file mode 100644
--- /dev/null
+++ b/Engine/Source/Runtime/AVEncoder/Public/AVEncoder.cpp
@@ -0,0 +1,212 @@
+// Copyright 1998-2019 Epic Games, Inc. All Rights Reserved.
+
+#include "AVEncoder.h"
+
+#if PLATFORM_WINDOWS
+	#include "Microsoft/Windows/NvVideoEncoder.h"
+	#include "Microsoft/Windows/AmfVideoEncoder.h"
+	#include "Microsoft/WmfAudioEncoder.h"
+#elif PLATFORM_XBOXONE
+	#include "Microsoft/WmfAudioEncoder.h"
+	#include "Microsoft/XBoxOne/XboxOneVideoEncoder.h"
+#endif
+
+namespace AVEncoder
+{
+
+bool GDefaultFactoriesRegistered = false;
+void RegisterDefaultFactories();
+
+void DoDefaultRegistration()
+{
+	if (GDefaultFactoriesRegistered)
+		return;
+	RegisterDefaultFactories();
+}
+
+//////////////////////////////////////////////////////////////////////////
+// FVideoEncoder
+//////////////////////////////////////////////////////////////////////////
+
+void FVideoEncoder::RegisterListener(IVideoEncoderListener& Listener)
+{
+	FScopeLock lock{ &ListenersMutex };
+	check(Listeners.Find(&Listener)==INDEX_NONE);
+	Listeners.AddUnique(&Listener);
+}
+
+void FVideoEncoder::UnregisterListener(IVideoEncoderListener& Listener)
+{
+	FScopeLock lock{ &ListenersMutex };
+	int32 Count = Listeners.Remove(&Listener);
+	check(Count == 1);
+}
+
+void FVideoEncoder::OnEncodedVideoFrame(const FAVPacket& Packet, TUniquePtr<FEncoderVideoFrameCookie> Cookie)
+{
+	FScopeLock lock{ &ListenersMutex };
+	for (auto&& L : Listeners)
+	{
+		L->OnEncodedVideoFrame(Packet, Cookie.Get());
+	}
+}
+
+//////////////////////////////////////////////////////////////////////////
+// FAudioEncoder
+//////////////////////////////////////////////////////////////////////////
+
+void FAudioEncoder::RegisterListener(IAudioEncoderListener& Listener)
+{
+	FScopeLock lock{ &ListenersMutex };
+	check(Listeners.Find(&Listener)==INDEX_NONE);
+	Listeners.AddUnique(&Listener);
+}
+
+void FAudioEncoder::UnregisterListener(IAudioEncoderListener& Listener)
+{
+	FScopeLock lock{ &ListenersMutex };
+	int32 Count = Listeners.Remove(&Listener);
+	check(Count == 1);
+}
+
+void FAudioEncoder::OnEncodedAudioFrame(const FAVPacket& Packet)
+{
+	FScopeLock lock{ &ListenersMutex };
+	for (auto&& L : Listeners)
+	{
+		L->OnEncodedAudioFrame(Packet);
+	}
+}
+
+//////////////////////////////////////////////////////////////////////////
+// FVideoEncoderFactory
+//////////////////////////////////////////////////////////////////////////
+
+TArray<FVideoEncoderFactory*> FVideoEncoderFactory::Factories;
+
+void FVideoEncoderFactory::RegisterFactory(FVideoEncoderFactory& Factory)
+{
+	DoDefaultRegistration();
+
+	Factories.AddUnique(&Factory);
+}
+
+void FVideoEncoderFactory::UnregisterFactory(FVideoEncoderFactory& Factory)
+{
+	Factories.Remove(&Factory);
+}
+
+FVideoEncoderFactory* FVideoEncoderFactory::FindFactory(const FString& Codec)
+{
+	DoDefaultRegistration();
+
+	for (auto&& Factory : Factories)
+	{
+		if (Factory->GetSupportedCodecs().Find(Codec) != INDEX_NONE)
+		{
+			return Factory;
+		}
+	}
+
+	return nullptr;
+}
+
+const TArray<FVideoEncoderFactory*> FVideoEncoderFactory::GetAllFactories()
+{
+	DoDefaultRegistration();
+
+	return Factories;
+}
+
+//////////////////////////////////////////////////////////////////////////
+// FAudioEncoderFactory
+//////////////////////////////////////////////////////////////////////////
+
+TArray<FAudioEncoderFactory*> FAudioEncoderFactory::Factories;
+
+void FAudioEncoderFactory::RegisterFactory(FAudioEncoderFactory& Factory)
+{
+	DoDefaultRegistration();
+
+	Factories.AddUnique(&Factory);
+}
+
+void FAudioEncoderFactory::UnregisterFactory(FAudioEncoderFactory& Factory)
+{
+	Factories.Remove(&Factory);
+}
+
+FAudioEncoderFactory* FAudioEncoderFactory::FindFactory(const FString& Codec)
+{
+	DoDefaultRegistration();
+
+	for (auto&& Factory : Factories)
+	{
+		if (Factory->GetSupportedCodecs().Find(Codec) != INDEX_NONE)
+		{
+			return Factory;
+		}
+	}
+
+	return nullptr;
+}
+
+const TArray<FAudioEncoderFactory*> FAudioEncoderFactory::GetAllFactories()
+{
+	DoDefaultRegistration();
+
+	return Factories;
+}
+
+void RegisterDefaultFactories()
+{
+	// We need to set this at the top, otherwise RegisterFactory will call this recursively
+	GDefaultFactoriesRegistered = true;
+
+#if PLATFORM_WINDOWS
+	// Nvidia NvEnc
+	static FNvVideoEncoderFactory NvVideoEncoderFactory;
+	FVideoEncoderFactory::RegisterFactory(NvVideoEncoderFactory);
+
+	// AMD Amf
+	static FAmfVideoEncoderFactory AmfVideoEncoderFactory;
+	FAmfVideoEncoderFactory::RegisterFactory(AmfVideoEncoderFactory);
+
+	// Generic Windows/XBox Wmf
+	static FWmfAudioEncoderFactory WmfAudioEncoderFactory;
+	FAudioEncoderFactory::RegisterFactory(WmfAudioEncoderFactory);
+
+#elif PLATFORM_XBOXONE
+	static FXboxOneVideoEncoderFactory XboxOneVideoEncoderFactory;
+	FVideoEncoderFactory::RegisterFactory(XboxOneVideoEncoderFactory);
+
+	static FWmfAudioEncoderFactory WmfAudioEncoderFactory;
+	FAudioEncoderFactory::RegisterFactory(WmfAudioEncoderFactory);
+#endif
+
+	// Log all available encoders
+	{
+		auto CodecsInfo = [&](auto&& Factories) -> FString
+		{
+			FString Str;
+			for (auto&& Factory : Factories)
+			{
+				Str += FString::Printf(TEXT(", %s(%s) "), Factory->GetName(), *FString::Join(Factory->GetSupportedCodecs(), TEXT("/")));
+			}
+			if (Str.IsEmpty())
+			{
+				return FString(TEXT("None"));
+			}
+			else
+			{
+				return Str;
+			}
+		};
+
+		UE_LOG(LogAVEncoder, Log, TEXT("Available video encoders: %s "), *CodecsInfo(FVideoEncoderFactory::GetAllFactories()));
+		UE_LOG(LogAVEncoder, Log, TEXT("Available audio encoders: %s "), *CodecsInfo(FAudioEncoderFactory::GetAllFactories()));
+	}
+}
+
+} // namespace AVEncoder
+
diff --git a/Engine/Source/Runtime/AVEncoder/Public/AVEncoder.h b/Engine/Source/Runtime/AVEncoder/Public/AVEncoder.h
new file mode 100644
--- /dev/null
+++ b/Engine/Source/Runtime/AVEncoder/Public/AVEncoder.h
@@ -0,0 +1,264 @@
+// Copyright 1998-2019 Epic Games, Inc. All Rights Reserved.
+
+#pragma once
+
+#include "CoreMinimal.h"
+#include "SampleBuffer.h"
+#include "RHI.h"
+
+namespace AVEncoder
+{
+
+struct FEncoderVideoFrameCookie
+{
+	virtual ~FEncoderVideoFrameCookie() {}
+};
+
+struct FVideoEncoderConfig
+{
+	uint32 Width = 0;
+	uint32 Height = 0;
+	uint32 Framerate = 0;
+	uint32 MaxBitrate = 0;
+	uint32 Bitrate = 0;
+
+	enum class EPreset
+	{
+		LowLatency,
+		HighQuality
+	};
+
+	/**
+	* Provides an hint to the encoder of what's the intended use.
+	*/
+	EPreset Preset = EPreset::HighQuality;
+
+	/** 
+	* Options specific to an hardware vendor (e.g: AMF or NVidia), or specific to a codec (h264)
+	* At the moment there are the following options (name and value)
+	*
+	* "qp" :
+	*	H264 Quantization Parameter (0-51). Lower values result in better quality but higher bitrate.
+	*	Default value is 20
+	* "ratecontrolmode" :
+	*	H264 rate control mode. Valid values are (constqp, vbr, cbr).
+	*	Default value is cbr
+	*/ 
+	TArray<TPair<FString, FString>> Options;
+};
+
+
+struct FAudioEncoderConfig
+{
+	uint32 Samplerate = 0;
+	uint32 NumChannels = 0;
+	// Encoder bitrate (bits per second, NOT bytes per second)
+	uint32 Bitrate = 0;
+	TArray<TPair<FString, FString>> Options;
+};
+
+struct FAudioFrame
+{
+	FTimespan Timestamp;
+	FTimespan Duration;
+	Audio::TSampleBuffer<float> Data;
+};
+
+enum class EPacketType
+{
+	Audio,
+	Video,
+	Invalid
+};
+
+struct FAVPacket
+{
+	explicit FAVPacket(EPacketType TypeIn)
+	{
+		Type = TypeIn;
+		if (Type == EPacketType::Audio)
+		{
+		}
+		else if (Type == EPacketType::Video)
+		{
+			FMemory::Memzero(Video);
+		}
+	}
+
+	FAVPacket(const FAVPacket& Other) = default;
+	FAVPacket(FAVPacket&& Other) = default;
+	FAVPacket& operator=(const FAVPacket& Other) = default;
+	FAVPacket& operator=(FAVPacket&& Other) = default;
+
+	EPacketType Type;
+	FTimespan Timestamp;
+	FTimespan Duration;
+	struct
+	{
+		FTimespan EncodeStartTs;
+		FTimespan EncodeFinishTs;
+	} Timings;
+	TArray<uint8> Data;
+
+	union
+	{
+		struct
+		{
+			bool bKeyFrame;
+			int32 Width;
+			int32 Height;
+			uint32 FrameAvgQP;
+			uint32 Framerate;
+		} Video;
+
+		struct
+		{
+			// Nothing at the moment
+		} Audio;
+	};
+
+	bool IsVideoKeyFrame() const
+	{
+		return (Type==EPacketType::Video && Video.bKeyFrame) ? true : false;
+	}
+
+	/**
+	 * The encoder can fail to encode a given frame
+	 * This might be expanded to an enum to show a reason, or removed entirely in the future
+	 * once we make the encoders a bit sturdier. Ideally, failing to encode should be dealt internally by the encoder,
+	 * correcting whatever failed. In turn a FVPacket received from the encoder should all be valid
+	 */
+	bool IsValid() const
+	{
+		return Data.Num() ? true : false;
+	}
+
+private:
+
+};
+
+class FVideoEncoder;
+class FAudioEncoder;
+
+class AVENCODER_API FVideoEncoderFactory
+{
+public:
+	virtual ~FVideoEncoderFactory() {}
+	virtual const TCHAR* GetName() const = 0;
+	virtual TArray<FString> GetSupportedCodecs() const = 0;
+	virtual TUniquePtr<FVideoEncoder> CreateEncoder(const FString& Codec) = 0;
+
+	static void RegisterFactory(FVideoEncoderFactory& Factory);
+	static void UnregisterFactory(FVideoEncoderFactory& Factory);
+	static FVideoEncoderFactory* FindFactory(const FString& Codec);
+	static const TArray<FVideoEncoderFactory*> GetAllFactories();
+
+private:
+
+	static TArray<FVideoEncoderFactory*> Factories;
+};
+
+class AVENCODER_API FAudioEncoderFactory
+{
+public:
+	virtual ~FAudioEncoderFactory() {}
+	virtual const TCHAR* GetName() const = 0;
+	virtual TArray<FString> GetSupportedCodecs() const = 0;
+	virtual TUniquePtr<FAudioEncoder> CreateEncoder(const FString& Codec) = 0;
+
+	static void RegisterFactory(FAudioEncoderFactory& Factory);
+	static void UnregisterFactory(FAudioEncoderFactory& Factory);
+	static FAudioEncoderFactory* FindFactory(const FString& Codec);
+	static const TArray<FAudioEncoderFactory*> GetAllFactories();
+
+private:
+	static TArray<FAudioEncoderFactory*> Factories;
+};
+
+class IVideoEncoderListener
+{
+public:
+	virtual void OnEncodedVideoFrame(const FAVPacket& Packet, FEncoderVideoFrameCookie* Cookie) = 0;
+};
+
+class IAudioEncoderListener
+{
+public:
+	virtual void OnEncodedAudioFrame(const FAVPacket& Packet) = 0;
+};
+
+using FBufferId = uint32;
+
+class AVENCODER_API FVideoEncoder
+{
+public:
+	FVideoEncoder() {}
+	virtual ~FVideoEncoder() {}
+	virtual const TCHAR* GetName() const = 0;
+	virtual const TCHAR* GetType() const = 0;
+	virtual bool Initialize(const FVideoEncoderConfig& Config) = 0;
+	virtual void Shutdown() = 0;
+
+	//
+	// To account for how webrtc works (for PixelStreaming), encoding a frame is a two step process.
+	// 1st - CopyTexture initiates a copy if the texture to the internal buffers, and returns an Id the caller can use to reference that internal buffer.
+	// 2nd - Either a Drop or Encode is required for each CopyTexture call, otherwise the respective internal buffer will stay
+	// @param Texture Texture to copy
+	// @param CaptureTs Capture timestamp
+	// @param Duration Delta time from the previous frame
+	// @param OutBufferId Buffer id on return
+	// @param Resolution If {0,0}, the copy will be the same size as the passed texture. If not {0,0}, it will use that specified resolution
+	// marked as used.
+	virtual bool CopyTexture(FTexture2DRHIRef Texture, FTimespan CaptureTs, FTimespan Duration, FBufferId& OutBufferId, FIntPoint Resolution = {0, 0}) = 0;
+	virtual void Drop(FBufferId BufferId) = 0;
+	virtual void Encode(FBufferId BufferId, bool bForceKeyFrame, uint32 Bitrate = 0, TUniquePtr<AVEncoder::FEncoderVideoFrameCookie> Cookie = nullptr) = 0;
+
+	// We return a copy instead of a reference, to simplify multithread  access
+	virtual FVideoEncoderConfig GetConfig() const = 0;
+
+	// Allows setting bitrate and framerate at runtime.
+	virtual bool SetBitrate(uint32 Bitrate) = 0;
+	virtual bool SetFramerate(uint32 Framerate) = 0;
+	virtual bool SetParameter(const FString& Parameter, const FString& Value) = 0;
+
+	virtual void RegisterListener(IVideoEncoderListener& Listener);
+	virtual void UnregisterListener(IVideoEncoderListener& Listener);
+
+protected:
+	void OnEncodedVideoFrame(const FAVPacket& Packet, TUniquePtr<FEncoderVideoFrameCookie> Cookie);
+
+private:
+
+	FCriticalSection ListenersMutex;
+	TArray<IVideoEncoderListener*> Listeners;
+
+};
+
+class AVENCODER_API FAudioEncoder
+{
+public:
+	virtual ~FAudioEncoder() {}
+	virtual const TCHAR* GetName() const = 0;
+	virtual const TCHAR* GetType() const = 0;
+	virtual bool Initialize(const FAudioEncoderConfig& Config) = 0;
+	virtual void Shutdown() = 0;
+
+	virtual void Encode(const FAudioFrame& Frame) = 0;
+	virtual FAudioEncoderConfig GetConfig() const = 0;
+
+	virtual void RegisterListener(IAudioEncoderListener& Listener);
+	virtual void UnregisterListener(IAudioEncoderListener& Listener);
+
+protected:
+	void OnEncodedAudioFrame(const FAVPacket& Packet);
+private:
+
+	FCriticalSection ListenersMutex;
+	TArray<IAudioEncoderListener*> Listeners;
+
+};
+
+void UnregisterDefaultFactories();
+
+}
+
diff --git a/Engine/Source/Runtime/AVEncoder/Public/Microsoft/AVEncoderIMFSampleWrapper.h b/Engine/Source/Runtime/AVEncoder/Public/Microsoft/AVEncoderIMFSampleWrapper.h
new file mode 100644
--- /dev/null
+++ b/Engine/Source/Runtime/AVEncoder/Public/Microsoft/AVEncoderIMFSampleWrapper.h
@@ -0,0 +1,116 @@
+// Copyright 1998-2019 Epic Games, Inc. All Rights Reserved.
+
+#pragma once
+
+#include "AVEncoder.h"
+#include "Templates/RefCounting.h"
+
+#if PLATFORM_WINDOWS
+
+#include "Windows/AllowWindowsPlatformTypes.h"
+#include "Windows/PreWindowsApi.h"
+	#include <mftransform.h>
+#include "Windows/PostWindowsApi.h"
+#include "Windows/HideWindowsPlatformTypes.h"
+
+#elif PLATFORM_XBOXONE
+
+#include "XboxOne/XboxOneAllowPlatformTypes.h"
+#include "XboxOne/XboxOnePreApi.h"
+	#include <mftransform.h>
+#include "XboxOne/XboxOnePostApi.h"
+#include "XboxOne/XboxOneHidePlatformTypes.h"
+#endif
+
+namespace AVEncoder
+{
+
+//
+// Wrapper for IMFSample, to make it easier to report errors
+//
+class AVENCODER_API FIMFSampleWrapper
+{
+public:
+
+	FIMFSampleWrapper(EPacketType InMediaType = EPacketType::Invalid, IMFSample* InSample = nullptr)
+		: MediaType(InMediaType)
+		, Sample(InSample)
+	{
+	}
+
+	const IMFSample* GetSample() const
+	{
+		return Sample;
+	}
+
+	IMFSample* GetSample()
+	{
+		return Sample;
+	}
+
+	bool CreateSample();
+
+	FTimespan GetTime() const;
+
+	void SetTime(FTimespan Time);
+
+	FTimespan GetDuration() const;
+
+	void SetDuration(FTimespan Duration);
+
+	bool IsVideoKeyFrame() const;
+
+	int GetBufferCount() const;
+
+	/**
+	*
+	* Calls the specified function for each buffer the sample contains
+	* Signature is "int (int BufferIndex, TArrayView<uint8> Data)
+	* The specified function should return true if iteration should continue, or false to finish
+	*
+	* The return value is what the specified function returned:
+	*	true : Iterated through all the buffers
+	*	false : Early termination
+	*/
+	template<typename T>
+	bool IterateBuffers(T&& Func)
+	{
+		int BufferCount = GetBufferCount();
+		for (int Idx = 0; Idx < BufferCount; ++Idx)
+		{
+			TRefCountPtr<IMFMediaBuffer> MediaBuffer = nullptr;
+			verify(SUCCEEDED(Sample->GetBufferByIndex(0, MediaBuffer.GetInitReference())));
+
+			BYTE* SrcData = nullptr;
+			DWORD MaxLength = 0;
+			DWORD CurrentLength = 0;
+			verify(SUCCEEDED(MediaBuffer->Lock(&SrcData, &MaxLength, &CurrentLength)));
+			bool res = Func(Idx, TArrayView<uint8>(SrcData, CurrentLength));
+			verify(SUCCEEDED(MediaBuffer->Unlock()));
+			if (!res)
+			{
+				return false;
+			}
+		}
+		return true;
+	}
+
+	bool IsValid() const
+	{
+		return Sample.IsValid();
+	}
+
+	void Reset()
+	{
+		Sample = nullptr;
+	}
+
+	FIMFSampleWrapper Clone() const;
+
+private:
+	EPacketType MediaType;
+	TRefCountPtr<IMFSample> Sample;
+};
+
+}
+
diff --git a/Engine/Source/Runtime/GameplayMediaEncoder/GameplayMediaEncoder.Build.cs b/Engine/Source/Runtime/GameplayMediaEncoder/GameplayMediaEncoder.Build.cs
--- a/Engine/Source/Runtime/GameplayMediaEncoder/GameplayMediaEncoder.Build.cs
+++ b/Engine/Source/Runtime/GameplayMediaEncoder/GameplayMediaEncoder.Build.cs
@@ -24,6 +24,7 @@ public class GameplayMediaEncoder : ModuleRules
 			"Slate",
 			"HTTP",
 			"Json",
+			"AVEncoder",
 			//"IBMRTMPIngest"
         });
 
@@ -40,6 +41,7 @@ public class GameplayMediaEncoder : ModuleRules
 
 			PublicSystemLibraries.Add("d3d11.lib");
 		}
+
 	}
 }
 
diff --git a/Engine/Source/Runtime/GameplayMediaEncoder/Private/GameplayMediaEncoder.cpp b/Engine/Source/Runtime/GameplayMediaEncoder/Private/GameplayMediaEncoder.cpp
--- a/Engine/Source/Runtime/GameplayMediaEncoder/Private/GameplayMediaEncoder.cpp
+++ b/Engine/Source/Runtime/GameplayMediaEncoder/Private/GameplayMediaEncoder.cpp
@@ -12,19 +12,6 @@
 #include "ProfilingDebugging/CsvProfiler.h"
 #include "IbmLiveStreaming.h"
 
-#if PLATFORM_WINDOWS || PLATFORM_XBOXONE
-	#include "Microsoft/WmfAudioEncoder.h"
-#endif
-
-#if PLATFORM_WINDOWS
-	#include "Microsoft/Windows/EncoderDevice.h"
-	#include "Microsoft/Windows/NvVideoEncoder.h"
-	#include "Microsoft/Windows/WmfVideoEncoder.h"
-	#include "Microsoft/Windows/AmdAmfVideoEncoder.h"
-#elif PLATFORM_XBOXONE
-	#include "Microsoft/XboxOne/XboxOneVideoEncoder.h"
-#endif
-
 DEFINE_LOG_CATEGORY(GameplayMediaEncoder);
 CSV_DEFINE_CATEGORY(GameplayMediaEncoder, true);
 
@@ -36,7 +23,7 @@ const uint32 HardcodedAudioSamplerate = 48000;
 const uint32 HardcodedAudioNumChannels = 2;
 // currently neither IVideoRecordingSystem neither HighlightFeature APIs allow to configure
 // audio stream parameters
-const uint32 HardcodedAudioBitrate = 24000;
+const uint32 HardcodedAudioBitrate = 192000;
 
 // currently neither IVideoRecordingSystem neither HighlightFeature APIs allow to configure
 // video stream parameters
@@ -54,8 +41,6 @@ const uint32 MaxVideoFPS = 60;
 const uint32 MaxWidth = 1920;
 const uint32 MaxHeight = 1080;
 
-GAMEPLAYMEDIAENCODER_START
-
 FAutoConsoleCommand GameplayMediaEncoderInitialize(
 	TEXT("GameplayMediaEncoder.Initialize"),
 	TEXT("Constructs the audio/video encoding objects. Does not start encoding"),
@@ -86,21 +71,24 @@ FAutoConsoleCommand GameplayMediaEncoderShutdown(
 //
 //////////////////////////////////////////////////////////////////////////
 
+FGameplayMediaEncoder* FGameplayMediaEncoder::Singleton = nullptr;
+
 FGameplayMediaEncoder* FGameplayMediaEncoder::Get()
 {
-	static FGameplayMediaEncoder Singleton;
-	return &Singleton;
+	if (!Singleton)
+	{
+		Singleton = new FGameplayMediaEncoder();
+	}
+	return Singleton;
 }
 
 FGameplayMediaEncoder::FGameplayMediaEncoder()
 {
-#if PLATFORM_WINDOWS
-	EncoderDevice = MakeShared<FEncoderDevice>();
-#endif
 }
 
 FGameplayMediaEncoder::~FGameplayMediaEncoder()
 {
+	Shutdown();
 }
 
 bool FGameplayMediaEncoder::RegisterListener(IGameplayMediaEncoderListener* Listener)
@@ -126,11 +114,8 @@ void FGameplayMediaEncoder::UnregisterListener(IGameplayMediaEncoderListener* Li
 	check(IsInGameThread());
 
 	ListenersCS.Lock();
-
 	Listeners.Remove(Listener);
-
 	bool bAnyListenersLeft = Listeners.Num() > 0;
-
 	ListenersCS.Unlock();
 
 	if (bAnyListenersLeft == false)
@@ -140,16 +125,6 @@ void FGameplayMediaEncoder::UnregisterListener(IGameplayMediaEncoderListener* Li
 	}
 }
 
-bool FGameplayMediaEncoder::GetAudioOutputType(TRefCountPtr<IMFMediaType>& OutType)
-{
-	return AudioEncoder ? AudioEncoder->GetOutputType(OutType) : false;
-}
-
-bool FGameplayMediaEncoder::GetVideoOutputType(TRefCountPtr<IMFMediaType>& OutType)
-{
-	return VideoEncoder ? VideoEncoder->GetOutputType(OutType) : false;
-}
-
 bool FGameplayMediaEncoder::Initialize()
 {
 	MemoryCheckpoint("Initial");
@@ -160,33 +135,51 @@ bool FGameplayMediaEncoder::Initialize()
 		return true;
 	}
 
-	auto OnMediaSampleReadyCallback = [this](const FGameplayMediaEncoderSample& Sample)
+	// If some error occurs, call Shutdown to cleanup
+	bool bIsOk = false;
+	ON_SCOPE_EXIT
 	{
-		return OnMediaSampleReady(Sample);
+		if (!bIsOk)
+		{
+			Shutdown();
+		}
 	};
 
 	//
 	// Audio
-	// 
-	AudioEncoder.Reset(new FWmfAudioEncoder(OnMediaSampleReadyCallback));
-	FWmfAudioEncoderConfig AudioConfig;
-	AudioConfig.SampleRate = HardcodedAudioSamplerate;
+	//
+	AVEncoder::FAudioEncoderFactory* AudioEncoderFactory = AVEncoder::FAudioEncoderFactory::FindFactory("aac");
+	if (!AudioEncoderFactory)
+	{
+		UE_LOG(GameplayMediaEncoder, Error, TEXT("No audio encoder for aac found"));
+		return false;
+	}
+
+	AudioEncoder = AudioEncoderFactory->CreateEncoder("aac");
+	if (!AudioEncoder)
+	{
+		UE_LOG(GameplayMediaEncoder, Error, TEXT("Could not create audio encoder"));
+		return false;
+	}
+
+	AVEncoder::FAudioEncoderConfig AudioConfig;
+	AudioConfig.Samplerate = HardcodedAudioSamplerate;
 	AudioConfig.NumChannels = HardcodedAudioNumChannels;
 	AudioConfig.Bitrate = HardcodedAudioBitrate;
 	if (!AudioEncoder->Initialize(AudioConfig))
 	{
-		Stop();
+		UE_LOG(GameplayMediaEncoder, Error, TEXT("Could not initialize audio encoder"));
 		return false;
 	}
 
-	bAudioFormatChecked = false;
+	AudioEncoder->RegisterListener(*this);
+
 	MemoryCheckpoint("Audio encoder initialized");
 
 	//
 	// Video
 	//
-	FVideoEncoderConfig VideoConfig;
-
+	AVEncoder::FVideoEncoderConfig VideoConfig;
 	FParse::Value(FCommandLine::Get(), TEXT("GameplayMediaEncoder.ResY="), VideoConfig.Height);
 	if (VideoConfig.Height == 0 || VideoConfig.Height == 720)
 	{
@@ -227,46 +220,37 @@ bool FGameplayMediaEncoder::Initialize()
 
 	UE_LOG(GameplayMediaEncoder, Log, TEXT("Using a config of {Width=%u, Height=%u, Framerate=%u, Bitrate=%u}"), VideoConfig.Width, VideoConfig.Height, VideoConfig.Framerate, VideoConfig.Bitrate);
 
-	//
-	// Try to initialize the best encoder for the current platform
-	//
-#if PLATFORM_WINDOWS
-	VideoEncoder.Reset(new FNvVideoEncoder(OnMediaSampleReadyCallback, EncoderDevice));
+	AVEncoder::FVideoEncoderFactory* VideoEncoderFactory = AVEncoder::FVideoEncoderFactory::FindFactory("h264");
+	if (!VideoEncoderFactory)
+	{
+		UE_LOG(GameplayMediaEncoder, Error, TEXT("No encoder for h264 found"));
+		return false;
+	}
+
+	VideoEncoder = VideoEncoderFactory->CreateEncoder("h264");
+
+	if (!VideoEncoder)
+	{
+		UE_LOG(GameplayMediaEncoder, Error, TEXT("Could not create video encoder"));
+		return false;
+	}
+
 	if (!VideoEncoder->Initialize(VideoConfig))
 	{
-		VideoEncoder.Reset(new FAmdAmfVideoEncoder(OnMediaSampleReadyCallback));
-		if (!VideoEncoder->Initialize(VideoConfig))
-		{
-			VideoEncoder.Reset(new FWmfVideoEncoder(OnMediaSampleReadyCallback));
-			if (!VideoEncoder->Initialize(VideoConfig))
-			{
-				Stop();
-				return false;
-			}
-		}
+		UE_LOG(GameplayMediaEncoder, Error, TEXT("Could not initialize video encoder"));
+		return false;
 	}
-#elif PLATFORM_XBOXONE
-	VideoEncoder.Reset(new FXboxOneVideoEncoder(OnMediaSampleReadyCallback));
-	verify(VideoEncoder->Initialize(VideoConfig));
-#endif
 
-	StartTime = 0;
-	AudioClock = 0;
+	VideoEncoder->RegisterListener(*this);
 
 	MemoryCheckpoint("Video encoder initialized");
 
+	bIsOk = true; // So Shutdown is not called due to the ON_SCOPE_EXIT
 	return true;
 }
 
 bool FGameplayMediaEncoder::Start()
 {
-	// Initialize if not initialized
-	if (!VideoEncoder)
-	{
-		if (!Initialize())
-			return false;
-	}
-
 	if (StartTime != 0)
 	{
 		UE_LOG(GameplayMediaEncoder, Log, TEXT("Already running"));
@@ -285,7 +269,7 @@ bool FGameplayMediaEncoder::Start()
 	StartTime = FTimespan::FromSeconds(FPlatformTime::Seconds());
 	AudioClock = 0;
 	NumCapturedFrames = 0;
-
+	
 	//
 	// subscribe to engine delegates for audio output and back buffer
 	//
@@ -293,10 +277,10 @@ bool FGameplayMediaEncoder::Start()
 	FAudioDevice* AudioDevice = GEngine->GetMainAudioDevice();
 	if (AudioDevice)
 	{
+		bAudioFormatChecked = false;
 		AudioDevice->RegisterSubmixBufferListener(this);
 	}
 
-	VideoEncoder->Start();
 	FSlateApplication::Get().GetRenderer()->OnBackBufferReadyToPresent().AddRaw(this, &FGameplayMediaEncoder::OnBackBufferReady);
 
 	return true;
@@ -326,7 +310,6 @@ void FGameplayMediaEncoder::Stop()
 		}
 	}
 
-	VideoEncoder->Stop();
 	StartTime = 0;
 	AudioClock = 0;
 }
@@ -341,28 +324,21 @@ void FGameplayMediaEncoder::Shutdown()
 
 	{
 		FScopeLock Lock(&AudioProcessingCS);
-		AudioEncoder.Reset();
+		if (AudioEncoder)
+		{
+			AudioEncoder->Shutdown();
+			AudioEncoder.Reset();
+		}
 	}
 	{
 		FScopeLock Lock(&VideoProcessingCS);
-		VideoEncoder.Reset();
-	}
-}
-
-bool FGameplayMediaEncoder::OnMediaSampleReady(const FGameplayMediaEncoderSample& Sample)
-{
-	FScopeLock ProcessMediaSamplesLock(&ProcessMediaSamplesCS);
-	if (bProcessMediaSamples)
-	{
-		FScopeLock Lock(&ListenersCS);
-
-		for (auto&& Listener : Listeners)
+		if (VideoEncoder)
 		{
-			Listener->OnMediaSample(Sample);
+			VideoEncoder->Shutdown();
+			VideoEncoder.Reset();
 		}
 	}
 
-	return true;
 }
 
 FTimespan FGameplayMediaEncoder::GetMediaTimestamp() const
@@ -394,7 +370,7 @@ void FGameplayMediaEncoder::OnBackBufferReady(SWindow& SlateWindow, const FTextu
 	ProcessVideoFrame(BackBuffer);
 }
 
-bool FGameplayMediaEncoder::ProcessAudioFrame(const float* AudioData, int32 NumSamples, int32 NumChannels, int32 SampleRate)
+void FGameplayMediaEncoder::ProcessAudioFrame(const float* AudioData, int32 NumSamples, int32 NumChannels, int32 SampleRate)
 {
 	Audio::AlignedFloatBuffer InData;
 	InData.Append(AudioData, NumSamples);
@@ -406,11 +382,6 @@ bool FGameplayMediaEncoder::ProcessAudioFrame(const float* AudioData, int32 NumS
 		FloatBuffer.MixBufferToChannels(HardcodedAudioNumChannels);
 	}
 
-	// First we need to clamp to [-1.0, 1.0], because the conversion to 16 bits doesn't clamp it,
-	// and we can end up with wrapping when playing loud sounds.
-	FloatBuffer.Clamp();
-	PCM16 = FloatBuffer;
-
 	// Adjust the AudioClock if for some reason it falls behind real time. This can happen if the game spikes, or if we break into the debugger.
 	FTimespan Now = GetMediaTimestamp();
 	if (AudioClock < Now.GetTotalSeconds())
@@ -420,19 +391,17 @@ bool FGameplayMediaEncoder::ProcessAudioFrame(const float* AudioData, int32 NumS
 		AudioClock = Now.GetTotalSeconds() + (FloatBuffer.GetSampleDuration() / 2);
 	}
 
-	{
-		FScopeLock Lock(&AudioProcessingCS);
-		FTimespan Timestamp = FTimespan::FromSeconds(AudioClock);
-		FTimespan Duration = FTimespan::FromSeconds(FloatBuffer.GetSampleDuration());
-		AudioEncoder->Process(reinterpret_cast<const int8*>(PCM16.GetData()), PCM16.GetNumSamples() * sizeof(*PCM16.GetData()), Timestamp, Duration);
-	}
+	AVEncoder::FAudioFrame Frame;
+	Frame.Timestamp = FTimespan::FromSeconds(AudioClock);
+	Frame.Duration = FTimespan::FromSeconds(FloatBuffer.GetSampleDuration());
+	FloatBuffer.Clamp();
+	Frame.Data = FloatBuffer;
+	AudioEncoder->Encode(Frame);
 
 	AudioClock += FloatBuffer.GetSampleDuration();
-
-	return true;
 }
 
-bool FGameplayMediaEncoder::ProcessVideoFrame(const FTexture2DRHIRef& BackBuffer)
+void FGameplayMediaEncoder::ProcessVideoFrame(const FTexture2DRHIRef& BackBuffer)
 {
 	FScopeLock Lock(&VideoProcessingCS);
 
@@ -445,24 +414,25 @@ bool FGameplayMediaEncoder::ProcessVideoFrame(const FTexture2DRHIRef& BackBuffer
 		if (NumCapturedFrames + 1 > NumExpectedFrames)
 		{
 			UE_LOG(GameplayMediaEncoder, Verbose, TEXT("Framerate control dropped captured frame"));
-			return true;
+			return;
 		}
 	}
 
 	if (!ChangeVideoConfig())
 	{
-		return false;
+		return;
 	}
 
-	FTimespan Duration = Now - LastVideoInputTimestamp;
-	if (!VideoEncoder->Process(BackBuffer, Now, Duration))
+	AVEncoder::FBufferId BufferId;
+	if (!VideoEncoder->CopyTexture(BackBuffer, Now, Now - LastVideoInputTimestamp, BufferId))
 	{
-		return false;
+		return;
 	}
 
+	VideoEncoder->Encode(BufferId, false, 0, nullptr);
+
 	LastVideoInputTimestamp = Now;
 	NumCapturedFrames++;
-	return true;
 }
 
 void FGameplayMediaEncoder::SetVideoBitrate(uint32 Bitrate)
@@ -497,12 +467,56 @@ bool FGameplayMediaEncoder::ChangeVideoConfig()
 			return false;
 		}
 		bChangeFramerate = false;
-		// reset framerate control
-		FramerateMonitoringStart = GetMediaTimestamp();
 		NumCapturedFrames = 0;
 	}
 
 	return true;
 }
 
-GAMEPLAYMEDIAENCODER_END
+void FGameplayMediaEncoder::OnEncodedAudioFrame(const AVEncoder::FAVPacket& Packet)
+{
+	OnEncodedFrame(Packet);
+}
+
+void FGameplayMediaEncoder::OnEncodedVideoFrame(const AVEncoder::FAVPacket& Packet, AVEncoder::FEncoderVideoFrameCookie* Cookie)
+{
+	OnEncodedFrame(Packet);
+}
+
+void FGameplayMediaEncoder::OnEncodedFrame(const AVEncoder::FAVPacket& Packet)
+{
+	FScopeLock Lock(&ListenersCS);
+	for (auto&& Listener : Listeners)
+	{
+		Listener->OnMediaSample(Packet);
+	}
+}
+
+TPair<FString, AVEncoder::FAudioEncoderConfig> FGameplayMediaEncoder::GetAudioConfig() const
+{
+	if (AudioEncoder)
+	{
+		return TPair<FString,AVEncoder::FAudioEncoderConfig>(
+			AudioEncoder->GetType(),
+			AudioEncoder->GetConfig());
+	}
+	else
+	{
+		return {};
+	}
+}
+
+TPair<FString, AVEncoder::FVideoEncoderConfig> FGameplayMediaEncoder::GetVideoConfig() const
+{
+	if (AudioEncoder)
+	{
+		return TPair<FString,AVEncoder::FVideoEncoderConfig>(
+			VideoEncoder->GetType(),
+			VideoEncoder->GetConfig());
+	}
+	else
+	{
+		return {};
+	}
+}
+
diff --git a/Engine/Source/Runtime/GameplayMediaEncoder/Private/GameplayMediaEncoderCommon.cpp b/Engine/Source/Runtime/GameplayMediaEncoder/Private/GameplayMediaEncoderCommon.cpp
--- a/Engine/Source/Runtime/GameplayMediaEncoder/Private/GameplayMediaEncoderCommon.cpp
+++ b/Engine/Source/Runtime/GameplayMediaEncoder/Private/GameplayMediaEncoderCommon.cpp
@@ -3,41 +3,6 @@
 #include "GameplayMediaEncoderCommon.h"
 #include "RHI.h"
 
-#if WMFMEDIA_SUPPORTED_PLATFORM
-	#pragma comment(lib, "mfplat")
-	#pragma comment(lib, "mfuuid")
-	#pragma comment(lib, "Mfreadwrite")
-#endif
-
-GAMEPLAYMEDIAENCODER_START
-
-//
-// Windows only code
-//
-#if PLATFORM_WINDOWS
-ID3D11Device* GetUE4DxDevice()
-{
-	auto Device = static_cast<ID3D11Device*>(GDynamicRHI->RHIGetNativeDevice());
-	checkf(Device != nullptr, TEXT("Failed to get UE4's ID3D11Device"));
-	return Device;
-}
-#endif
-
-//
-// XboxOne only code
-// 
-#if PLATFORM_XBOXONE 
-ID3D12Device* GetUE4DxDevice()
-{
-	auto Device = static_cast<ID3D12Device*>(GDynamicRHI->RHIGetNativeDevice());
-	checkf(Device != nullptr, TEXT("Failed to get UE4's ID3D12Device"));
-	return Device;
-}
-
-#endif
-
-
-// #RVF : Remove these once the code is production ready
 TArray<FMemoryCheckpoint> gMemoryCheckpoints;
 uint64 MemoryCheckpoint(const FString& Name)
 {
@@ -77,6 +42,3 @@ void LogMemoryCheckpoints(const FString& Name)
 	}
 }
 
-
-GAMEPLAYMEDIAENCODER_END
-
diff --git a/Engine/Source/Runtime/GameplayMediaEncoder/Private/GameplayMediaEncoderCommon.h b/Engine/Source/Runtime/GameplayMediaEncoder/Private/GameplayMediaEncoderCommon.h
--- a/Engine/Source/Runtime/GameplayMediaEncoder/Private/GameplayMediaEncoderCommon.h
+++ b/Engine/Source/Runtime/GameplayMediaEncoder/Private/GameplayMediaEncoderCommon.h
@@ -18,17 +18,9 @@
 
 //
 // Macros to control some things during development
-// #RVF : Remove this when cleaning up
 //
-#define WRITE_TO_FILE 0
 #define LIVESTREAMING 0
 
-#if UE_BUILD_SHIPPING
-	#define GAMEPLAYMEDIAENCODER_DEBUG 0
-#else
-	#define GAMEPLAYMEDIAENCODER_DEBUG 0
-#endif
-
 //
 // Includes common to Windows and XboxOne
 //
@@ -36,242 +28,10 @@
 	#include "Templates/RefCounting.h"
 #endif
 
-//
-// Windows only include
-//
-#if PLATFORM_WINDOWS
-
-	#include "Windows/AllowWindowsPlatformTypes.h"
-	#include "Windows/PreWindowsApi.h"
-		#include <d3d11.h>
-		#include <mftransform.h>
-		#include <mfapi.h>
-		#include <mferror.h>
-		#include <mfidl.h>
-		#include <codecapi.h>
-		#include <shlwapi.h>
-		#include <mfreadwrite.h>
-	#include "Windows/PostWindowsApi.h"
-	#include "Windows/HideWindowsPlatformTypes.h"
-
-	#include "D3D11State.h"
-	#include "D3D11Resources.h"
-#endif
-
-//
-// XboxOne only includes
-//
-#if PLATFORM_XBOXONE
-
-#include "XboxOne/XboxOneAllowPlatformTypes.h"
-#include "XboxOne/XboxOnePreApi.h"
-		#include <d3d11_x.h>
-		#include <d3d12_x.h>
-		#include <d3dcompiler_x.h>
-		#include <d3dx12_x.h>
-		#include <mftransform.h>
-		#include <mfapi.h>
-		#include <mferror.h>
-		#include <mfidl.h>
-		#include <codecapi.h>
-		#include <mfreadwrite.h>
-#include "XboxOne/XboxOnePostApi.h"
-#include "XboxOne/XboxOneHidePlatformTypes.h"
-
-#endif
-
-#if GAMEPLAYMEDIAENCODER_DEBUG
-	#define GAMEPLAYMEDIAENCODER_START PRAGMA_DISABLE_OPTIMIZATION
-	#define GAMEPLAYMEDIAENCODER_END PRAGMA_ENABLE_OPTIMIZATION
-#else
-	#define GAMEPLAYMEDIAENCODER_START 
-	#define GAMEPLAYMEDIAENCODER_END 
-#endif
-
 #define WMFMEDIA_SUPPORTED_PLATFORM (PLATFORM_WINDOWS && (WINVER >= 0x0600 /*Vista*/) && !UE_SERVER)
 
 DECLARE_LOG_CATEGORY_EXTERN(GameplayMediaEncoder, Log, VeryVerbose);
 
-//
-// Windows and XboxOne common code/macros
-//
-#if PLATFORM_WINDOWS || PLATFORM_XBOXONE
-
-inline const FString GetComErrorDescription(HRESULT Res)
-{
-	const uint32 BufSize = 4096;
-	WIDECHAR buffer[4096];
-	if (::FormatMessageW(FORMAT_MESSAGE_FROM_SYSTEM,
-		nullptr,
-		Res,
-		MAKELANGID(LANG_NEUTRAL, SUBLANG_NEUTRAL),
-		buffer,
-		sizeof(buffer) / sizeof(*buffer),
-		nullptr))
-	{
-		return buffer;
-	}
-	else
-	{
-		return TEXT("[cannot find error description]");
-	}
-}
-
-#if PLATFORM_WINDOWS
-	#include "Windows/AllowWindowsPlatformTypes.h"
-#elif PLATFORM_XBOXONE
-	#include "XboxOne/XboxOneAllowPlatformTypes.h"
-#endif
-
-// macro to deal with COM calls inside a function that returns `false` on error
-#define CHECK_HR(COM_call)\
-	{\
-		HRESULT Res = COM_call;\
-		if (FAILED(Res))\
-		{\
-			UE_LOG(GameplayMediaEncoder, Error, TEXT("`" #COM_call "` failed: 0x%X - %s"), Res, *GetComErrorDescription(Res));\
-			return false;\
-		}\
-	}
-
-// macro to deal with COM calls inside COM method (that returns HRESULT)
-#define CHECK_HR_COM(COM_call)\
-	{\
-		HRESULT Res = COM_call;\
-		if (FAILED(Res))\
-		{\
-			UE_LOG(GameplayMediaEncoder, Error, TEXT("`" #COM_call "` failed: 0x%X - %s"), Res, *GetComErrorDescription(Res));\
-			return Res;\
-		}\
-	}
-
-// macro to deal with COM calls inside COM method (that simply returns)
-#define CHECK_HR_VOID(COM_call)\
-	{\
-		HRESULT Res = COM_call;\
-		if (FAILED(Res))\
-		{\
-			UE_LOG(GameplayMediaEncoder, Error, TEXT("`" #COM_call "` failed: 0x%X - %s"), Res, *GetComErrorDescription(Res));\
-			return;\
-		}\
-	}
-
-#if PLATFORM_WINDOWS
-	#include "Windows/HideWindowsPlatformTypes.h"
-#elif PLATFORM_XBOXONE
-	#include "XboxOne/XboxOneHidePlatformTypes.h"
-#endif
-
-// following commented include causes name clash between UE4 and Windows `IMediaEventSink`,
-// we just need a couple of GUIDs from there so the solution is to duplicate them below
-//#include "wmcodecdsp.h"
-
-const GUID CLSID_AACMFTEncoder = { 0x93AF0C51, 0x2275, 0x45d2, { 0xA3, 0x5B, 0xF2, 0xBA, 0x21, 0xCA, 0xED, 0x00 } };
-const GUID CLSID_CMSH264EncoderMFT = { 0x6ca50344, 0x051a, 0x4ded, { 0x97, 0x79, 0xa4, 0x33, 0x05, 0x16, 0x5e, 0x35 } };
-const GUID CLSID_VideoProcessorMFT = { 0x88753b26, 0x5b24, 0x49bd, { 0xb2, 0xe7, 0xc, 0x44, 0x5c, 0x78, 0xc9, 0x82 } };
-
-// `MF_LOW_LATENCY` is defined in "mfapi.h" for >= WIN8
-// UE4 supports lower Windows versions at the moment and so `WINVER` is < `_WIN32_WINNT_WIN8`
-// to be able to use `MF_LOW_LATENCY` with default UE4 build we define it ourselves and check actual
-// Windows version in runtime
-#if (WINVER < _WIN32_WINNT_WIN8)
-	const GUID MF_LOW_LATENCY = { 0x9c27891a, 0xed7a, 0x40e1,{ 0x88, 0xe8, 0xb2, 0x27, 0x27, 0xa0, 0x24, 0xee } };
-#endif
-
-
-#endif
-
-//
-// Windows only utility code/macros
-//
-#if PLATFORM_WINDOWS
-
-ID3D11Device* GetUE4DxDevice();
-
-// scope-disable particular DX11 Debug Layer errors
-class FScopeDisabledDxDebugErrors final
-{
-private:
-
-public:
-	FScopeDisabledDxDebugErrors(TArray<D3D11_MESSAGE_ID>&& ErrorsToDisable)
-	{
-		TRefCountPtr<ID3D11Debug> Debug;
-		HRESULT HRes = GetUE4DxDevice()->QueryInterface(__uuidof(ID3D11Debug), reinterpret_cast<void**>(Debug.GetInitReference()));
-
-		if (HRes == E_NOINTERFACE)
-		{
-			// Debug Layer is not enabled, so no need to disable its errors
-			return;
-		}
-
-		if (!SUCCEEDED(HRes) ||
-			!SUCCEEDED(HRes = Debug->QueryInterface(__uuidof(ID3D11InfoQueue), reinterpret_cast<void**>(InfoQueue.GetInitReference()))))
-		{
-			UE_LOG(GameplayMediaEncoder, VeryVerbose, TEXT("Failed to get ID3D11InfoQueue: 0x%X - %s"), HRes, *GetComErrorDescription(HRes));
-			return;
-		}
-
-		D3D11_INFO_QUEUE_FILTER filter = {};
-		filter.DenyList.NumIDs = ErrorsToDisable.Num();
-		filter.DenyList.pIDList = ErrorsToDisable.GetData();
-		bSucceeded = SUCCEEDED(InfoQueue->PushStorageFilter(&filter));
-	}
-
-	~FScopeDisabledDxDebugErrors()
-	{
-		if (bSucceeded)
-		{
-			InfoQueue->PopStorageFilter();
-		}
-	}
-
-private:
-	TRefCountPtr<ID3D11InfoQueue> InfoQueue;
-	bool bSucceeded = false;
-};
-
-
-#endif
-
-//
-// XboxOne only utility code/macros
-//
-#if PLATFORM_XBOXONE
-
-ID3D12Device* GetUE4DxDevice();
-
-#endif
-
-template<typename F>
-inline void ExecuteRHICommand(F&& Functor)
-{
-	FRHICommandList& RHICmdList = FRHICommandListExecutor::GetImmediateCommandList();
-	if (RHICmdList.Bypass())
-	{
-		Functor();
-	}
-	else
-	{
-		FRHICOMMAND_MACRO(FLocalRHICommand)
-		{
-			F Functor;
-
-			FLocalRHICommand(F InFunctor)
-				: Functor(MoveTemp(InFunctor))
-			{}
-
-			void Execute(FRHICommandListBase& CmdList)
-			{
-				Functor();
-			}
-		};
-
-		new (RHICmdList.AllocCommand<FLocalRHICommand>()) FLocalRHICommand(Functor);
-	}
-}
-
-// #RVF : Remove this once the code is production ready.
 struct FMemoryCheckpoint
 {
 	FString Name;
diff --git a/Engine/Source/Runtime/GameplayMediaEncoder/Private/GameplayMediaEncoderModule.cpp b/Engine/Source/Runtime/GameplayMediaEncoder/Private/GameplayMediaEncoderModule.cpp
--- a/Engine/Source/Runtime/GameplayMediaEncoder/Private/GameplayMediaEncoderModule.cpp
+++ b/Engine/Source/Runtime/GameplayMediaEncoder/Private/GameplayMediaEncoderModule.cpp
@@ -4,8 +4,7 @@
 #include "Modules/ModuleInterface.h"
 #include "Modules/ModuleManager.h"
 #include "GameplayMediaEncoderCommon.h"
-
-GAMEPLAYMEDIAENCODER_START
+#include "GameplayMediaEncoder.h"
 
 class FGameplayMediaEncoderModule : public IModuleInterface
 {
@@ -14,8 +13,28 @@ public:
 	{
 	}
 
+	~FGameplayMediaEncoderModule()
+	{
+	}
+
+	void StartupModule() override
+	{
+		FModuleManager::Get().LoadModule(TEXT("AVEncoder"));
+	}
+
+	void ShutdownModule() override
+	{
+		// If the FGameplayMediaEncoder instance was created, then explicitly destroy it here
+		// instead of waiting for the automatic cleanup, since at that point some objects
+		// it depends to for a clean shutdown are not available any longer.
+		if (FGameplayMediaEncoder::Singleton)
+		{
+			delete FGameplayMediaEncoder::Singleton;
+			FGameplayMediaEncoder::Singleton = nullptr;
+		}
+	}
+private:
 };
 
 IMPLEMENT_MODULE(FGameplayMediaEncoderModule, GameplayMediaEncoder);
 
-GAMEPLAYMEDIAENCODER_END
diff --git a/Engine/Source/Runtime/GameplayMediaEncoder/Private/IbmLiveStreaming.cpp b/Engine/Source/Runtime/GameplayMediaEncoder/Private/IbmLiveStreaming.cpp
--- a/Engine/Source/Runtime/GameplayMediaEncoder/Private/IbmLiveStreaming.cpp
+++ b/Engine/Source/Runtime/GameplayMediaEncoder/Private/IbmLiveStreaming.cpp
@@ -15,8 +15,6 @@
 
 DEFINE_LOG_CATEGORY(IbmLiveStreaming);
 
-GAMEPLAYMEDIAENCODER_START
-
 extern "C"
 {
 	// rtmp_c_static.lib needs us to define this symbol
@@ -1018,7 +1016,5 @@ void FIbmLiveStreaming::OnStreamBandwidthChanged(RTMPModuleBroadcaster* Module,
 	reinterpret_cast<FIbmLiveStreaming*>(Module->customData)->OnStreamBandwidthChangedImpl(Module, Bandwidth, QueueWasEmpty != 0 );
 }
 
-GAMEPLAYMEDIAENCODER_END
-
 #endif // WITH_IBMRTMPINGEST
 
diff --git a/Engine/Source/Runtime/GameplayMediaEncoder/Public/GameplayMediaEncoder.h b/Engine/Source/Runtime/GameplayMediaEncoder/Public/GameplayMediaEncoder.h
--- a/Engine/Source/Runtime/GameplayMediaEncoder/Public/GameplayMediaEncoder.h
+++ b/Engine/Source/Runtime/GameplayMediaEncoder/Public/GameplayMediaEncoder.h
@@ -6,22 +6,20 @@
 
 #include "Logging/LogMacros.h"
 #include "AudioMixerDevice.h"
-#include "GameplayMediaEncoderSample.h"
+#include "AVEncoder.h"
 
 #include "RHI.h"
 #include "RHIResources.h"
 
-class FBaseVideoEncoder;
-class FWmfAudioEncoder;
 class SWindow;
 
 class IGameplayMediaEncoderListener
 {
 public:
-	virtual void OnMediaSample(const FGameplayMediaEncoderSample& Sample) = 0;
+	virtual void OnMediaSample(const AVEncoder::FAVPacket& Sample) = 0;
 };
 
-class GAMEPLAYMEDIAENCODER_API FGameplayMediaEncoder final : private ISubmixBufferListener
+class GAMEPLAYMEDIAENCODER_API FGameplayMediaEncoder final : private ISubmixBufferListener, private AVEncoder::IAudioEncoderListener, private AVEncoder::IVideoEncoderListener
 {
 public:
 
@@ -30,18 +28,20 @@ public:
 	 */
 	static FGameplayMediaEncoder* Get();
 
-	FGameplayMediaEncoder();
 	~FGameplayMediaEncoder();
 
 	bool RegisterListener(IGameplayMediaEncoderListener* Listener);
 	void UnregisterListener(IGameplayMediaEncoderListener* Listener);
 
-	bool GetAudioOutputType(TRefCountPtr<IMFMediaType>& OutType);
-	bool GetVideoOutputType(TRefCountPtr<IMFMediaType>& OutType);
-
 	void SetVideoBitrate(uint32 Bitrate);
 	void SetVideoFramerate(uint32 Framerate);
 
+	/**
+	 * Returns the audio codec name and configuration
+	 */
+	TPair<FString, AVEncoder::FAudioEncoderConfig> GetAudioConfig() const;
+	TPair<FString, AVEncoder::FVideoEncoderConfig> GetVideoConfig() const;
+
 	bool Initialize();
 	void Shutdown();
 	bool Start();
@@ -72,39 +72,43 @@ public:
 	}
 
 private:
-	FTimespan GetMediaTimestamp() const;
 
-	bool OnMediaSampleReady(const FGameplayMediaEncoderSample& Sample);
+	// Private to control how our single instance is created
+	FGameplayMediaEncoder();
+
+	// Returns how long it has been recording for.
+	FTimespan GetMediaTimestamp() const;
 
 	// Back buffer capture
 	void OnBackBufferReady(SWindow& SlateWindow, const FTexture2DRHIRef& BackBuffer);
 	// ISubmixBufferListener interface
 	void OnNewSubmixBuffer(const USoundSubmix* OwningSubmix, float* AudioData, int32 NumSamples, int32 NumChannels, const int32 SampleRate, double AudioClock) override;
 
-	bool ProcessAudioFrame(const float* AudioData, int32 NumSamples, int32 NumChannels, int32 SampleRate);
-	bool ProcessVideoFrame(const FTexture2DRHIRef& BackBuffer);
+	void ProcessAudioFrame(const float* AudioData, int32 NumSamples, int32 NumChannels, int32 SampleRate);
+	void ProcessVideoFrame(const FTexture2DRHIRef& BackBuffer);
 
 	bool ChangeVideoConfig();
 
+	//
+	// AVEncoder::IAudioEncoderListener interface
+	void OnEncodedAudioFrame(const AVEncoder::FAVPacket& Packet) override;
+	//
+	// AVEncoder::IVideoEncoderListener interface
+	void OnEncodedVideoFrame(const AVEncoder::FAVPacket& Packet, AVEncoder::FEncoderVideoFrameCookie* Cookie) override;
+
+	void OnEncodedFrame(const AVEncoder::FAVPacket& Packet);
+
 	FCriticalSection ListenersCS;
 	TArray<IGameplayMediaEncoderListener*> Listeners;
 
 	FCriticalSection AudioProcessingCS;
 	FCriticalSection VideoProcessingCS;
-
-	TUniquePtr<FWmfAudioEncoder> AudioEncoder;
-	TUniquePtr<FBaseVideoEncoder> VideoEncoder;
-#if PLATFORM_WINDOWS
-	TSharedPtr<class FEncoderDevice> EncoderDevice;
-#endif
-	bool bAudioFormatChecked = false;
-	bool bDoFrameSkipping = false;
-
-	// Keep this as a member variables, to reuse the memory allocation
-	Audio::TSampleBuffer<int16> PCM16;
+	TUniquePtr<AVEncoder::FAudioEncoder> AudioEncoder;
+	TUniquePtr<AVEncoder::FVideoEncoder> VideoEncoder;
 
 	uint64 NumCapturedFrames = 0;
 	FTimespan StartTime = 0;
+
 	// Instead of using the AudioClock parameter ISubmixBufferListener::OnNewSubmixBuffer gives us, we calculate our own, by
 	// advancing it as we receive more data.
 	// This is so that we can adjust the clock if things get out of sync, such as if we break into the debugger.
@@ -112,16 +116,17 @@ private:
 
 	FTimespan LastVideoInputTimestamp = 0;
 
-	// It is possible to suspend the processing of media samples which is
-	// required during resolution change.
-	FCriticalSection ProcessMediaSamplesCS;
-	bool bProcessMediaSamples = true;
+	bool bAudioFormatChecked = false;
+	bool bDoFrameSkipping = false;
+
+	friend class FGameplayMediaEncoderModule;
+	static FGameplayMediaEncoder* Singleton;
 
 	// live streaming: quality adaptation to available uplink b/w
 	TAtomic<uint32> NewVideoBitrate{ 0 };
 	FThreadSafeBool bChangeBitrate = false;
-	FTimespan FramerateMonitoringStart = -1;
 	TAtomic<uint32> NewVideoFramerate{ 0 };
 	FThreadSafeBool bChangeFramerate = false;
+
 };
 
diff --git a/Engine/Source/Runtime/Windows/WindowsPlatformFeatures/Private/HighlightRecorder.cpp b/Engine/Source/Runtime/Windows/WindowsPlatformFeatures/Private/HighlightRecorder.cpp
--- a/Engine/Source/Runtime/Windows/WindowsPlatformFeatures/Private/HighlightRecorder.cpp
+++ b/Engine/Source/Runtime/Windows/WindowsPlatformFeatures/Private/HighlightRecorder.cpp
@@ -2,7 +2,6 @@
 
 #include "HighlightRecorder.h"
 #include "WmfMp4Writer.h"
-#include "GameplayMediaEncoderSample.h"
 
 #include "Misc/Paths.h"
 
@@ -131,7 +130,7 @@ FTimespan FHighlightRecorder::GetRecordingTime() const
 	return FTimespan::FromSeconds(FPlatformTime::Seconds()) - RecordingStartTime - TotalPausedDuration;
 }
 
-void FHighlightRecorder::OnMediaSample(const FGameplayMediaEncoderSample& InSample)
+void FHighlightRecorder::OnMediaSample(const AVEncoder::FAVPacket& InSample)
 {
 	// We might be paused, so don't do anything
 	if (State != EState::Recording)
@@ -140,7 +139,7 @@ void FHighlightRecorder::OnMediaSample(const FGameplayMediaEncoderSample& InSamp
 	}
 
 	// Only start pushing video frames once we receive a key frame
-	if (NumPushedFrames == 0 && InSample.GetType() == EMediaType::Video)
+	if (NumPushedFrames == 0 && InSample.Type == AVEncoder::EPacketType::Video)
 	{
 		if (!InSample.IsVideoKeyFrame())
 		{
@@ -150,11 +149,19 @@ void FHighlightRecorder::OnMediaSample(const FGameplayMediaEncoderSample& InSamp
 		++NumPushedFrames;
 	}
 
-	FGameplayMediaEncoderSample SampleCopy = InSample.Clone();
+	AVEncoder::FAVPacket SampleCopy = InSample;
+
+
+	AVEncoder::FAVPacket A = InSample;
+	AVEncoder::FAVPacket B = MoveTemp(A);
+	AVEncoder::FAVPacket C(AVEncoder::EPacketType::Video);
+	C = MoveTemp(B);
+	AVEncoder::FAVPacket D(AVEncoder::EPacketType::Video);
+	D = C;
 
 	if (TotalPausedDuration != 0)
 	{
-		SampleCopy.SetTime(SampleCopy.GetTime() - TotalPausedDuration);
+		SampleCopy.Timestamp = SampleCopy.Timestamp  - TotalPausedDuration;
 	}
 
 	RingBuffer.Push(MoveTemp(SampleCopy));
@@ -185,6 +192,11 @@ bool FHighlightRecorder::SaveHighlight(const TCHAR* Filename, FDoneCallback InDo
 
 	{
 		CSV_SCOPED_TIMING_STAT(WindowsVideoRecordingSystem, HighlightRecorder_SaveThreadCreation);
+		if (BackgroundSaving && BackgroundSaving->IsJoinable())
+		{
+			BackgroundSaving->Join();
+		}
+
 		BackgroundSaving.Reset(new FThread(TEXT("Highlight Saving"), [this, LocalFilename, MaxDurationSecs]()
 		{
 			SaveHighlightInBackground(LocalFilename, MaxDurationSecs);
@@ -216,7 +228,7 @@ bool FHighlightRecorder::SaveHighlightInBackground(const FString& Filename, doub
 
 bool FHighlightRecorder::SaveHighlightInBackgroundImpl(const FString& Filename, double MaxDurationSecs)
 {
-	TArray<FGameplayMediaEncoderSample> Samples = RingBuffer.GetCopy();
+	TArray<AVEncoder::FAVPacket> Samples = RingBuffer.GetCopy();
 
 	if (Samples.Num()==0)
 	{
@@ -235,7 +247,7 @@ bool FHighlightRecorder::SaveHighlightInBackgroundImpl(const FString& Filename,
 	bool bHasAudio = false;
 	for (int Idx = FirstSampleIndex; Idx != Samples.Num(); ++Idx)
 	{
-		if (Samples[Idx].GetType() == EMediaType::Audio)
+		if (Samples[Idx].Type == AVEncoder::EPacketType::Audio)
 		{
 			bHasAudio = true;
 			break;
@@ -247,7 +259,7 @@ bool FHighlightRecorder::SaveHighlightInBackgroundImpl(const FString& Filename,
 		return false;
 	}
 
-	checkf(Samples[FirstSampleIndex].IsVideoKeyFrame(), TEXT("t %.3f d %.3f"), Samples[FirstSampleIndex].GetTime().GetTotalSeconds(), Samples[FirstSampleIndex].GetDuration().GetTotalSeconds());
+	checkf(Samples[FirstSampleIndex].IsVideoKeyFrame(), TEXT("t %.3f d %.3f"), Samples[FirstSampleIndex].Timestamp.GetTotalSeconds(), Samples[FirstSampleIndex].Duration.GetTotalSeconds());
 
 	if (FirstSampleIndex == Samples.Num())
 	{
@@ -255,14 +267,18 @@ bool FHighlightRecorder::SaveHighlightInBackgroundImpl(const FString& Filename,
 		return false;
 	}
 
-	UE_LOG(HighlightRecorder, Verbose, TEXT("writting %d samples to .mp4, %.3f s, starting from %.3f s, index %d"), Samples.Num() - FirstSampleIndex, (Samples.Last().GetTime() - StartTime + Samples.Last().GetDuration()).GetTotalSeconds(), StartTime.GetTotalSeconds(), FirstSampleIndex);
+	UE_LOG(HighlightRecorder, Verbose, TEXT("writting %d samples to .mp4, %.3f s, starting from %.3f s, index %d"),
+		Samples.Num() - FirstSampleIndex,
+		(Samples.Last().Timestamp - StartTime + Samples.Last().Duration).GetTotalSeconds(),
+		StartTime.GetTotalSeconds(),
+		FirstSampleIndex);
 
 	// get samples starting from `StartTime` and push them into Mp4Writer
 	for (int Idx = FirstSampleIndex; Idx != Samples.Num(); ++Idx)
 	{
-		FGameplayMediaEncoderSample& Sample = Samples[Idx];
-		Sample.SetTime(Sample.GetTime() - StartTime);
-		if (!Mp4Writer->Write(Sample, (Sample.GetType()==EMediaType::Video && bHasAudio) ? 1 : 0))
+		AVEncoder::FAVPacket& Sample = Samples[Idx];
+		Sample.Timestamp = Sample.Timestamp - StartTime;
+		if (!Mp4Writer->Write(Sample, (Sample.Type==AVEncoder::EPacketType::Audio) ? AudioStreamIndex : VideoStreamIndex))
 		{
 			return false;
 		}
@@ -299,28 +315,39 @@ bool FHighlightRecorder::InitialiseMp4Writer(const FString& Filename, bool bHasA
 		return false;
 	}
 
-	TRefCountPtr<IMFMediaType> AudioType;
-	if (!FGameplayMediaEncoder::Get()->GetAudioOutputType(AudioType))
-	{
-		return false;
-	}
-
-	DWORD StreamIndex;
 	if (bHasAudio)
 	{
-		if (!Mp4Writer->CreateStream(AudioType, StreamIndex))
+		TPair<FString, AVEncoder::FAudioEncoderConfig> AudioConfig = FGameplayMediaEncoder::Get()->GetAudioConfig();
+		if (AudioConfig.Key == "")
+		{
+			UE_LOG(HighlightRecorder, Error, TEXT("Could not get audio config"));
+			return false;
+		}
+		
+		TOptional<DWORD> Res = Mp4Writer->CreateAudioStream(AudioConfig.Key, AudioConfig.Value);
+		if (Res.IsSet())
+		{
+			AudioStreamIndex = Res.GetValue();
+		}
+		else
 		{
 			return false;
 		}
 	}
 
-	TRefCountPtr<IMFMediaType> VideoType;
-	if (!FGameplayMediaEncoder::Get()->GetVideoOutputType(VideoType))
+	TPair<FString, AVEncoder::FVideoEncoderConfig> VideoConfig = FGameplayMediaEncoder::Get()->GetVideoConfig();
+	if (VideoConfig.Key == "")
 	{
+		UE_LOG(HighlightRecorder, Error, TEXT("Could not get video config"));
 		return false;
 	}
 
-	if (!Mp4Writer->CreateStream(VideoType, StreamIndex))
+	TOptional<DWORD> Res = Mp4Writer->CreateVideoStream(VideoConfig.Key, VideoConfig.Value);
+	if (Res.IsSet())
+	{
+		VideoStreamIndex = Res.GetValue();
+	}
+	else
 	{
 		return false;
 	}
@@ -334,7 +361,7 @@ bool FHighlightRecorder::InitialiseMp4Writer(const FString& Filename, bool bHasA
 }
 
 // finds index and timestamp of the first sample that should be written to .mp4
-bool FHighlightRecorder::GetSavingStart(const TArray<FGameplayMediaEncoderSample>& Samples, FTimespan MaxDuration, int& StartIndex, FTimespan& StartTime) const
+bool FHighlightRecorder::GetSavingStart(const TArray<AVEncoder::FAVPacket>& Samples, FTimespan MaxDuration, int& StartIndex, FTimespan& StartTime) const
 // the first sample in .mp4 file should have timestamp 0 and all other timestamps should be relative to the
 // first one
 // 1) if `MaxDurationSecs` > actual ring buffer duration (last sample timestamp - first) -> we need to save all
@@ -352,7 +379,7 @@ bool FHighlightRecorder::GetSavingStart(const TArray<FGameplayMediaEncoderSample
 		return false;
 	}
 
-	FTimespan FirstTimestamp = Samples[0].GetTime();
+	FTimespan FirstTimestamp = Samples[0].Timestamp;
 
 	if (FirstTimestamp > StartTime)
 	{
@@ -365,7 +392,7 @@ bool FHighlightRecorder::GetSavingStart(const TArray<FGameplayMediaEncoderSample
 	bool bFound = false;
 	for (; i != Samples.Num(); ++i)
 	{
-		FTimespan Time = Samples[i].GetTime();
+		FTimespan Time = Samples[i].Timestamp;
 		if (Time >= StartTime && Samples[i].IsVideoKeyFrame())
 		{
 			// correct StartTime to match timestamp of the first sample to be written
diff --git a/Engine/Source/Runtime/Windows/WindowsPlatformFeatures/Private/HighlightRecorder.h b/Engine/Source/Runtime/Windows/WindowsPlatformFeatures/Private/HighlightRecorder.h
--- a/Engine/Source/Runtime/Windows/WindowsPlatformFeatures/Private/HighlightRecorder.h
+++ b/Engine/Source/Runtime/Windows/WindowsPlatformFeatures/Private/HighlightRecorder.h
@@ -45,7 +45,7 @@ private:
 	bool SaveHighlightInBackground(const FString& Filename, double MaxDurationSecs);
 	bool SaveHighlightInBackgroundImpl(const FString& Filename, double MaxDurationSecs);
 	bool InitialiseMp4Writer(const FString& Filename, bool bHasAudio);
-	bool GetSavingStart(const TArray<FGameplayMediaEncoderSample>& Samples, FTimespan MaxDuration, int& OutStartIndex, FTimespan& OutStartTime) const;
+	bool GetSavingStart(const TArray<AVEncoder::FAVPacket>& Samples, FTimespan MaxDuration, int& OutStartIndex, FTimespan& OutStartTime) const;
 
 	// takes into account if we've been paused and shifts current time back to compensate paused state
 	// so all timestamps are continuous even over paused pieces
@@ -54,7 +54,7 @@ private:
 	//
 	// IGameplayMediaEncoderListener implementation
 	//
-	void OnMediaSample(const FGameplayMediaEncoderSample& Sample) override;
+	void OnMediaSample(const AVEncoder::FAVPacket& Sample) override;
 
 private:
 	TAtomic<EState> State{ EState::Stopped };
@@ -78,6 +78,9 @@ private:
 	FDoneCallback DoneCallback;
 	FThreadSafeBool bSaving = false;
 
+	DWORD AudioStreamIndex = 0;
+	DWORD VideoStreamIndex = 0;
+
 #pragma region testing
 public:
 	static void Start(const TArray<FString>& Args, UWorld*, FOutputDevice& Output)
diff --git a/Engine/Source/Runtime/Windows/WindowsPlatformFeatures/Private/WindowsPlatformFeatures.cpp b/Engine/Source/Runtime/Windows/WindowsPlatformFeatures/Private/WindowsPlatformFeatures.cpp
--- a/Engine/Source/Runtime/Windows/WindowsPlatformFeatures/Private/WindowsPlatformFeatures.cpp
+++ b/Engine/Source/Runtime/Windows/WindowsPlatformFeatures/Private/WindowsPlatformFeatures.cpp
@@ -11,8 +11,6 @@ WINDOWSPLATFORMFEATURES_START
 
 FWindowsPlatformFeaturesModule::FWindowsPlatformFeaturesModule()
 {
-	// load generic modules
-	StartupModules();
 }
 
 IVideoRecordingSystem* FWindowsPlatformFeaturesModule::GetVideoRecordingSystem()
@@ -21,10 +19,9 @@ IVideoRecordingSystem* FWindowsPlatformFeaturesModule::GetVideoRecordingSystem()
 	return &VideoRecordingSystem;
 }
 
-bool FWindowsPlatformFeaturesModule::StartupModules()
+void FWindowsPlatformFeaturesModule::StartupModule()
 {
 	FModuleManager::Get().LoadModule(TEXT("GameplayMediaEncoder"));
-	return true;
 }
 
 WINDOWSPLATFORMFEATURES_END
diff --git a/Engine/Source/Runtime/Windows/WindowsPlatformFeatures/Private/WindowsPlatformFeatures.h b/Engine/Source/Runtime/Windows/WindowsPlatformFeatures/Private/WindowsPlatformFeatures.h
--- a/Engine/Source/Runtime/Windows/WindowsPlatformFeatures/Private/WindowsPlatformFeatures.h
+++ b/Engine/Source/Runtime/Windows/WindowsPlatformFeatures/Private/WindowsPlatformFeatures.h
@@ -16,10 +16,7 @@ public:
 	virtual IVideoRecordingSystem* GetVideoRecordingSystem() override;
 
 private:
-	/**
-	 * Load global/generic modules, and perform any initialization
-	 */
-	bool StartupModules();
+	void StartupModule() override;
 
 };
 
