Compile fix (missing header)

diff --git a/Engine/Source/Runtime/AVEncoder/Private/AVEncoderCommon.cpp b/Engine/Source/Runtime/AVEncoder/Private/AVEncoderCommon.cpp
--- a/Engine/Source/Runtime/AVEncoder/Private/AVEncoderCommon.cpp
+++ b/Engine/Source/Runtime/AVEncoder/Private/AVEncoderCommon.cpp
@@ -1,6 +1,7 @@
 // Copyright 1998-2019 Epic Games, Inc. All Rights Reserved.
 
 #include "AVEncoderCommon.h"
+#include "Modules/ModuleManager.hh"
 
 DEFINE_LOG_CATEGORY(LogAVEncoder);
 CSV_DECLARE_CATEGORY_EXTERN(AVEncoder);
