Compile fix (clang is more strict than visual studio)

diff --git a/Engine/Source/Runtime/AVEncoder/Private/Microsoft/Windows/AmfVideoEncoder.cpp b/Engine/Source/Runtime/AVEncoder/Private/Microsoft/Windows/AmfVideoEncoder.cpp
--- a/Engine/Source/Runtime/AVEncoder/Private/Microsoft/Windows/AmfVideoEncoder.cpp
+++ b/Engine/Source/Runtime/AVEncoder/Private/Microsoft/Windows/AmfVideoEncoder.cpp
@@ -188,7 +188,7 @@ private:
 	struct FFrame
 	{
 		const FBufferId Id = 0;
-		TAtomic<EFrameState> State = EFrameState::Free;
+		TAtomic<EFrameState> State = { EFrameState::Free };
 		FInputFrame InputFrame;
 		FOutputFrame OutputFrame;
 		uint64 FrameIdx = 0;
