Fix PVS warning V607: Ownerless expression 'Config.Bitrate'.
#rnx

diff --git a/Engine/Source/Runtime/AVEncoder/Private/Microsoft/Windows/NvVideoEncoder.cpp b/Engine/Source/Runtime/AVEncoder/Private/Microsoft/Windows/NvVideoEncoder.cpp
--- a/Engine/Source/Runtime/AVEncoder/Private/Microsoft/Windows/NvVideoEncoder.cpp
+++ b/Engine/Source/Runtime/AVEncoder/Private/Microsoft/Windows/NvVideoEncoder.cpp
@@ -619,7 +619,7 @@ FVideoEncoderConfig FNvVideoEncoder::GetConfig() const
 
 bool FNvVideoEncoder::SetBitrate(uint32 Bitrate)
 {
-	Config.Bitrate;
+	Config.Bitrate = Bitrate;
 	return true;
 }
 
