Added #if around Windows & XboxOne code, so it's not compiled for any other platforms.

diff --git a/Engine/Source/Runtime/AVEncoder/Private/Microsoft/AVEncoderIMFSampleWrapper.cpp b/Engine/Source/Runtime/AVEncoder/Private/Microsoft/AVEncoderIMFSampleWrapper.cpp
--- a/Engine/Source/Runtime/AVEncoder/Private/Microsoft/AVEncoderIMFSampleWrapper.cpp
+++ b/Engine/Source/Runtime/AVEncoder/Private/Microsoft/AVEncoderIMFSampleWrapper.cpp
@@ -1,6 +1,9 @@
 // Copyright 1998-2019 Epic Games, Inc. All Rights Reserved.
 
 #include "Microsoft/AVEncoderIMFSampleWrapper.h"
+
+#if PLATFORM_WINDOWS || PLATFORM_XBOXONE
+
 #include "Microsoft/AVEncoderMicrosoftCommon.h"
 
 namespace AVEncoder
@@ -117,6 +120,5 @@ FIMFSampleWrapper FIMFSampleWrapper::Clone() const
 
 } // namespace AVEncoder
 
-
-
+#endif // PLATFORM_WINDOWS || PLATFORM_XBOXONE
 
diff --git a/Engine/Source/Runtime/AVEncoder/Private/Microsoft/AVEncoderMicrosoftCommon.cpp b/Engine/Source/Runtime/AVEncoder/Private/Microsoft/AVEncoderMicrosoftCommon.cpp
--- a/Engine/Source/Runtime/AVEncoder/Private/Microsoft/AVEncoderMicrosoftCommon.cpp
+++ b/Engine/Source/Runtime/AVEncoder/Private/Microsoft/AVEncoderMicrosoftCommon.cpp
@@ -4,6 +4,8 @@
 
 #include "AVEncoderMicrosoftCommon.h"
 
+#if PLATFORM_WINDOWS || PLATFORM_XBOXONE
+
 #if WMFMEDIA_SUPPORTED_PLATFORM
 	#pragma comment(lib, "mfplat")
 	#pragma comment(lib, "mfuuid")
@@ -41,3 +43,5 @@ ID3D12Device* GetUE4DxDevice()
 
 } // namespace AVEncoder
 
+#endif // PLATFORM_WINDOWS || PLATFORM_XBOXONE
+
diff --git a/Engine/Source/Runtime/AVEncoder/Private/Microsoft/WmfAudioEncoder.cpp b/Engine/Source/Runtime/AVEncoder/Private/Microsoft/WmfAudioEncoder.cpp
--- a/Engine/Source/Runtime/AVEncoder/Private/Microsoft/WmfAudioEncoder.cpp
+++ b/Engine/Source/Runtime/AVEncoder/Private/Microsoft/WmfAudioEncoder.cpp
@@ -1,9 +1,13 @@
 // Copyright 1998-2019 Epic Games, Inc. All Rights Reserved.
 
 #include "Microsoft/WmfAudioEncoder.h"
+
+#if PLATFORM_WINDOWS || PLATFORM_XBOXONE
+
 #include "Microsoft/AVEncoderIMFSampleWrapper.h"
 #include "Microsoft/AVEncoderMicrosoftCommon.h"
 
+
 namespace AVEncoder
 {
 
@@ -399,6 +403,7 @@ TUniquePtr<FAudioEncoder> FWmfAudioEncoderFactory::CreateEncoder(const FString&
 	}
 }
 
-
 }
 
+#endif
+
diff --git a/Engine/Source/Runtime/AVEncoder/Private/Microsoft/WmfAudioEncoder.h b/Engine/Source/Runtime/AVEncoder/Private/Microsoft/WmfAudioEncoder.h
--- a/Engine/Source/Runtime/AVEncoder/Private/Microsoft/WmfAudioEncoder.h
+++ b/Engine/Source/Runtime/AVEncoder/Private/Microsoft/WmfAudioEncoder.h
@@ -5,6 +5,8 @@
 #include "AVEncoder.h"
 #include "AVEncoderCommon.h"
 
+#if PLATFORM_WINDOWS || PLATFORM_XBOXONE
+
 namespace AVEncoder
 {
 
@@ -21,4 +23,5 @@ private:
 
 }
 
+#endif
 
diff --git a/Engine/Source/Runtime/AVEncoder/Public/AVEncoder.cpp b/Engine/Source/Runtime/AVEncoder/Public/AVEncoder.cpp
--- a/Engine/Source/Runtime/AVEncoder/Public/AVEncoder.cpp
+++ b/Engine/Source/Runtime/AVEncoder/Public/AVEncoder.cpp
@@ -172,14 +172,15 @@ void RegisterDefaultFactories()
 	static FAmfVideoEncoderFactory AmfVideoEncoderFactory;
 	FAmfVideoEncoderFactory::RegisterFactory(AmfVideoEncoderFactory);
 
-	// Generic Windows/XBox Wmf
-	static FWmfAudioEncoderFactory WmfAudioEncoderFactory;
-	FAudioEncoderFactory::RegisterFactory(WmfAudioEncoderFactory);
+#else PLATFORM_XBOXONE
 
-#elif PLATFORM_XBOXONE
 	static FXboxOneVideoEncoderFactory XboxOneVideoEncoderFactory;
 	FVideoEncoderFactory::RegisterFactory(XboxOneVideoEncoderFactory);
 
+#endif
+
+#if PLATFORM_WINDOWS || PLATFORM_XBOXONE
+	// Generic Windows/XBox Wmf encoder
 	static FWmfAudioEncoderFactory WmfAudioEncoderFactory;
 	FAudioEncoderFactory::RegisterFactory(WmfAudioEncoderFactory);
 #endif
diff --git a/Engine/Source/Runtime/AVEncoder/Public/Microsoft/AVEncoderIMFSampleWrapper.h b/Engine/Source/Runtime/AVEncoder/Public/Microsoft/AVEncoderIMFSampleWrapper.h
--- a/Engine/Source/Runtime/AVEncoder/Public/Microsoft/AVEncoderIMFSampleWrapper.h
+++ b/Engine/Source/Runtime/AVEncoder/Public/Microsoft/AVEncoderIMFSampleWrapper.h
@@ -5,6 +5,8 @@
 #include "AVEncoder.h"
 #include "Templates/RefCounting.h"
 
+#if PLATFORM_WINDOWS || PLATFORM_XBOXONE
+
 #if PLATFORM_WINDOWS
 
 #include "Windows/AllowWindowsPlatformTypes.h"
@@ -114,3 +116,5 @@ private:
 
 }
 
+#endif // PLATFORM_WINDOWS || PLATFORM_XBOXONE
+
