Fixed compile error for non Windows platforms. (had a #else instead of a #elif)

diff --git a/Engine/Source/Runtime/AVEncoder/Public/AVEncoder.cpp b/Engine/Source/Runtime/AVEncoder/Public/AVEncoder.cpp
--- a/Engine/Source/Runtime/AVEncoder/Public/AVEncoder.cpp
+++ b/Engine/Source/Runtime/AVEncoder/Public/AVEncoder.cpp
@@ -172,7 +172,7 @@ void RegisterDefaultFactories()
 	static FAmfVideoEncoderFactory AmfVideoEncoderFactory;
 	FAmfVideoEncoderFactory::RegisterFactory(AmfVideoEncoderFactory);
 
-#else PLATFORM_XBOXONE
+#elif PLATFORM_XBOXONE
 
 	static FXboxOneVideoEncoderFactory XboxOneVideoEncoderFactory;
 	FVideoEncoderFactory::RegisterFactory(XboxOneVideoEncoderFactory);
