Missing #if

diff --git a/Engine/Source/Runtime/AVEncoder/Private/Microsoft/AVEncoderMicrosoftCommon.h b/Engine/Source/Runtime/AVEncoder/Private/Microsoft/AVEncoderMicrosoftCommon.h
--- a/Engine/Source/Runtime/AVEncoder/Private/Microsoft/AVEncoderMicrosoftCommon.h
+++ b/Engine/Source/Runtime/AVEncoder/Private/Microsoft/AVEncoderMicrosoftCommon.h
@@ -5,6 +5,8 @@
 #include "AVEncoder.h"
 #include "AVEncoderCommon.h"
 
+#if PLATFORM_WINDOWS || PLATFORM_XBOXONE
+
 //
 // Windows only include
 //
@@ -210,3 +212,7 @@ private:
 
 } // namespace AVEncoder
 
+
+#endif // PLATFORM_WINDOWS || PLATFORM_XBOXONE
+
+
