Missing #include

diff --git a/Engine/Source/Runtime/AVEncoder/Public/AVEncoder.cpp b/Engine/Source/Runtime/AVEncoder/Public/AVEncoder.cpp
--- a/Engine/Source/Runtime/AVEncoder/Public/AVEncoder.cpp
+++ b/Engine/Source/Runtime/AVEncoder/Public/AVEncoder.cpp
@@ -1,6 +1,7 @@
 // Copyright 1998-2019 Epic Games, Inc. All Rights Reserved.
 
 #include "AVEncoder.h"
+#include "AVEncoderCommon.h"
 
 #if PLATFORM_WINDOWS
 	#include "Microsoft/Windows/NvVideoEncoder.h"
