Compile fix for 'Compile ClangEditor Win64'

- Remove pragma once from the cpp

#rnx

diff --git a/Engine/Source/Runtime/AVEncoder/Private/Microsoft/AVEncoderMicrosoftCommon.cpp b/Engine/Source/Runtime/AVEncoder/Private/Microsoft/AVEncoderMicrosoftCommon.cpp
--- a/Engine/Source/Runtime/AVEncoder/Private/Microsoft/AVEncoderMicrosoftCommon.cpp
+++ b/Engine/Source/Runtime/AVEncoder/Private/Microsoft/AVEncoderMicrosoftCommon.cpp
@@ -1,7 +1,5 @@
 // Copyright 1998-2019 Epic Games, Inc. All Rights Reserved.
 
-#pragma once
-
 #include "AVEncoderMicrosoftCommon.h"
 
 #if PLATFORM_WINDOWS || PLATFORM_XBOXONE
