Bringing what I had in CL 11245130.

diff --git a/Engine/Source/Runtime/AVEncoder/Public/AVEncoder.h b/Engine/Source/Runtime/AVEncoder/Public/AVEncoder.h
--- a/Engine/Source/Runtime/AVEncoder/Public/AVEncoder.h
+++ b/Engine/Source/Runtime/AVEncoder/Public/AVEncoder.h
@@ -6,20 +6,122 @@
 #include "SampleBuffer.h"
 #include "RHI.h"
 
+/**
+ * The AVEncoder module provides a simple API to encode Video Audio.
+ * Initially, it was created to remove the duplicated encoding code done for Pixel Streaming
+ * and gameplay recording (GameplayMediaEncoder module), so at the moment it only supports what
+ * those systems require (h264 for video, and aac for audio)
+ * 
+ * Both audio and video encoders are exposed through factories, and the existing framework
+ * allows registering user made encoders.
+ * 
+ * The following is a simplified example, capturing gameplay video, ignoring errors for simplicity.
+ * For a full example on how to capture gameplay video and audio look at the GameplayMediaEncoder module.
+ * 
+ * <pre>
+ *
+ * class FFoo : private AVEncoder::IVideoEncoderListener
+ * {
+ * 	TUniquePtr<AVEncoder::FVideoEncoder> VideoEncoder;
+ * 	FTimespan StartTime = 0;
+ * 	FTimespan LastVideoFrameTimestamp = 0;
+ * 
+ * 	FFoo()
+ * 	{
+ * 	}
+ * 
+ * 	bool Initialize()
+ * 	{
+ * 		// Find a h264 video encoder (returns nullptr)
+ * 		AVEncoder::FVideoEncoderFactory* VideoEncoderFactory = AVEncoder::FVideoEncoderFactory::FindFactory("h264");
+ * 		if (!VideoEncoderFactory)
+ * 			return false;
+ * 
+ * 		// Create encoder
+ * 		VideoEncoder = VideoEncoderFactory->CreateEncoder("h264");
+ * 		if (!VideoEncoder)
+ * 			return false;
+ * 
+ * 		// Desired settings
+ * 		AVEncoder::FVideoEncoderConfig Config;
+ * 		Config.Width = 1920;
+ * 		Config.Height = 1080;
+ * 		Config.Framerate = 60;
+ * 		Config.Bitrate = 4000000; // 4mbps
+ * 		if (!VideoEncoder->Initialize(Config))
+ * 			return false;
+ * 
+ * 		// Register ourselves to receive the encoded frames
+ * 		VideoEncoder->RegisterListener(this);
+ * 
+ * 		// Register back buffer capture with UE4
+ * 		FSlateApplication::Get().GetRenderer()->OnBackBufferReadyToPresent().AddRaw(this, &FGameplayMediaEncoder::OnBackBufferReady);
+ * 
+ * 		StartTime = FTimespan::FromSeconds(FPlatformTime::Seconds());
+ * 
+ * 		return true;
+ * 	}
+ * 
+ * 
+ * 	// Back buffer capture callback from UE4
+ * 	void OnBackBufferReady(SWindow& SlateWindow, const FTexture2DRHIRef& BackBuffer)
+ * 	{
+ * 		FTimespan Now = FTimespan::FromSeconds(FPlatformTime::Seconds()) - StartTime;
+ * 		AVEncoder::FBufferId BufferId;
+ * 
+ * 		// Initialize copy of the backbuffer to the encoder.
+ * 		if (!VideoEncoder->CopyTexture(BackBuffer, Now, Now - LastVideoFameTimestamp, BufferId))
+ * 			return;
+ * 
+ * 		// Encode the frame
+ * 		VideoEncoder->Encode(BufferId, false);
+ * 		LastVideoInputTimestamp = Now;
+ * 	}
+ * 
+ * 	// AVEncoder::IVideoEncoderListener interface
+ * 	void OnEncodedVideoFrame(const AVEncoder::FAVPacket& Packet, AVEncoder::FEncoderVideoFrameCookie* Cookie) override
+ * 	{
+ * 		// Encoded frame is in `Packet`
+ * 	}
+ * 
+ * };
+ * 
+ * </pre>
+ */
 namespace AVEncoder
 {
 
+/**
+ * When sending a frame to the video encoder, it is possible to specify a cookie that will be passed back to the user code
+ * in the encoder callback.
+ * Whenever you have the need to associate some state to a given frame, you can derive from this, and your own state
+ */
 struct FEncoderVideoFrameCookie
 {
 	virtual ~FEncoderVideoFrameCookie() {}
 };
 
+/**
+ * Video encoder initial configuration.
+ */
 struct FVideoEncoderConfig
 {
+	/**
+	 * Encoding Resolution (e.g: 1920 x 1080 )
+	 */
 	uint32 Width = 0;
 	uint32 Height = 0;
 	uint32 Framerate = 0;
+
+	/**
+	 * Maximum bitrate in bps. Mostly an hint for the encoder.
+	 * Normally this is used for VBR and ignored for CBR
+	 */
 	uint32 MaxBitrate = 0;
+
+	/**
+	 * Target bitrate in bps
+	 */
 	uint32 Bitrate = 0;
 
 	enum class EPreset
@@ -48,6 +150,9 @@ struct FVideoEncoderConfig
 };
 
 
+/**
+ * Audio encoder initial configuration
+ */
 struct FAudioEncoderConfig
 {
 	uint32 Samplerate = 0;
@@ -71,6 +176,11 @@ enum class EPacketType
 	Invalid
 };
 
+/**
+ * Encoded output.
+ * Both FVideoEncoder and FAudioEncoder use this same type to output the encoded
+ * data.
+ */
 struct FAVPacket
 {
 	explicit FAVPacket(EPacketType TypeIn)
@@ -93,13 +203,25 @@ struct FAVPacket
 	EPacketType Type;
 	FTimespan Timestamp;
 	FTimespan Duration;
+
+	/**
+	 * Encoding latency
+	 */
 	struct
 	{
 		FTimespan EncodeStartTs;
 		FTimespan EncodeFinishTs;
 	} Timings;
+
+	/**
+	 * Actual encoded output
+	 */
 	TArray<uint8> Data;
 
+	/**
+	 * Depending if it's a Video or Audio packet, you should only access either
+	 * the Video or Audio member.
+	 */
 	union
 	{
 		struct
@@ -117,6 +239,9 @@ struct FAVPacket
 		} Audio;
 	};
 
+	/**
+	 * Returns true if this packet is a video key frame
+	 */
 	bool IsVideoKeyFrame() const
 	{
 		return (Type==EPacketType::Video && Video.bKeyFrame) ? true : false;
@@ -133,24 +258,51 @@ struct FAVPacket
 		return Data.Num() ? true : false;
 	}
 
-private:
-
 };
 
 class FVideoEncoder;
 class FAudioEncoder;
 
+/**
+ * Allows querying existing video encoders.
+ *
+ * It provides methods to query available encoders, register/unregister new ones, and create given encoders.
+ */
 class AVENCODER_API FVideoEncoderFactory
 {
 public:
 	virtual ~FVideoEncoderFactory() {}
+
+	/**
+	* Returns the factory name. Typically this is the SDK/hardware used. E.g: "amf" or "nvenc"
+	*/
 	virtual const TCHAR* GetName() const = 0;
+
+	/**
+	 * Returns all the encoders potentially available with this factory.
+	 * E.g: { "h264" }
+	 */
 	virtual TArray<FString> GetSupportedCodecs() const = 0;
+
+	/**
+	 * Creates the requested encoder, given the encoder name (e.g: "h264")
+	 */
 	virtual TUniquePtr<FVideoEncoder> CreateEncoder(const FString& Codec) = 0;
 
+	/**
+	 * These allow registering user supplied factories
+	 */
 	static void RegisterFactory(FVideoEncoderFactory& Factory);
 	static void UnregisterFactory(FVideoEncoderFactory& Factory);
+
+	/**
+	 *  Returns potentially best factory that can encode with the specified codec
+	 */
 	static FVideoEncoderFactory* FindFactory(const FString& Codec);
+
+	/**
+	 * Returns all currently registered factories
+	 */
 	static const TArray<FVideoEncoderFactory*> GetAllFactories();
 
 private:
@@ -158,6 +310,11 @@ private:
 	static TArray<FVideoEncoderFactory*> Factories;
 };
 
+/**
+ * Allows querying existing audio encoders.
+ *
+ * Member functions provide functionally equivalent to FVideoEncoderFactory
+ */
 class AVENCODER_API FAudioEncoderFactory
 {
 public:
@@ -175,52 +332,125 @@ private:
 	static TArray<FAudioEncoderFactory*> Factories;
 };
 
+/**
+ * Interface that should be implemented to react to video encoder's events
+ */
 class IVideoEncoderListener
 {
 public:
 	virtual void OnEncodedVideoFrame(const FAVPacket& Packet, FEncoderVideoFrameCookie* Cookie) = 0;
 };
 
+/**
+ * Interface that should be implemented to react to audio encoder's events
+ */
 class IAudioEncoderListener
 {
 public:
 	virtual void OnEncodedAudioFrame(const FAVPacket& Packet) = 0;
 };
 
+/**
+* Identifies the internal buffer being used by the video encoder.
+*/
 using FBufferId = uint32;
 
+/**
+ *  Base class for video encoders.
+ *
+ *  @note Threading model is still work in progress.
+ */
 class AVENCODER_API FVideoEncoder
 {
 public:
 	FVideoEncoder() {}
 	virtual ~FVideoEncoder() {}
+
+	/**
+	 * Returns the name of the encoder. This normally specifies the codec, and sdk/hardware.
+	 * E.g : "h264.nvenc", or "h264.amf"
+	 */
 	virtual const TCHAR* GetName() const = 0;
+
+	/**
+	 * Returns just the codec type. E.g: "h264"
+	 */
 	virtual const TCHAR* GetType() const = 0;
+
+	/**
+	 * Initialize the encoder with the specified initial configuration
+	 */
 	virtual bool Initialize(const FVideoEncoderConfig& Config) = 0;
+
+	/**
+	 * Shutdown MUST be called before destroying the encoder
+	 */
 	virtual void Shutdown() = 0;
 
-	//
-	// To account for how webrtc works (for PixelStreaming), encoding a frame is a two step process.
-	// 1st - CopyTexture initiates a copy if the texture to the internal buffers, and returns an Id the caller can use to reference that internal buffer.
-	// 2nd - Either a Drop or Encode is required for each CopyTexture call, otherwise the respective internal buffer will stay
-	// @param Texture Texture to copy
-	// @param CaptureTs Capture timestamp
-	// @param Duration Delta time from the previous frame
-	// @param OutBufferId Buffer id on return
-	// @param Resolution If {0,0}, the copy will be the same size as the passed texture. If not {0,0}, it will use that specified resolution
-	// marked as used.
+	/**
+	 * Copy the specified texture to an internal buffer.
+	 *
+	 * This is a bit awkward, but it's somewhat required to account for how webrtc works (for PixelStreaming), where 
+	 * encoding a frame is a two step process :
+	 * 1st - CopyTexture initiates a copy of the texture to the internal buffers, and returns an Id the caller can use to reference that internal buffer.
+	 * 2nd - Either a Drop or Encode is required for each CopyTexture call, otherwise the respective internal buffer will stay
+	 * @param Texture Texture to copy
+	 * @param CaptureTs Capture timestamp
+	 * @param Duration Delta time from the previous frame
+	 * @param OutBufferId Buffer id on return
+	 * @param Resolution If {0,0}, the copy will be the same size as the passed texture. If not {0,0}, it will use that specified resolution
+	 * marked as used.
+	 *
+	 * @note This MUST be called from the render thread
+	 */
 	virtual bool CopyTexture(FTexture2DRHIRef Texture, FTimespan CaptureTs, FTimespan Duration, FBufferId& OutBufferId, FIntPoint Resolution = {0, 0}) = 0;
+
+	/**
+	 * Drops an internal buffer that was returned by CopyTexture.
+	 * This should be used if the application decides that it doesn't want to encode a frame that was already
+	 * passed to CopyTexture.
+	 * Once this is called on a valid buffer id, do not use it that buffer id again.
+	 *
+	 * @param BufferId Buffer returned by a CopyTexture call
+	 *
+	 */
 	virtual void Drop(FBufferId BufferId) = 0;
+
+	/**
+	 * Initiates the encoding of the given buffer.
+	 * Once this is called on a valid buffer, do not use it again.
+
+	 * @param BufferId Buffer returned by a CopyTexture call
+	 * @param bForceKeyFrame Force this frame to be encoded as a key frame
+	 * @param Bitrate Target bitrate at the moment this frame is sent to the encoder.
+	 *	This allows tweaking bitrate at runtime
+	 * @param Cookie Allows storing custom per-frame state if the application needs it. 
+	 *	This is returned back to the application when it receives the encoded frame.
+	 *
+	 */
 	virtual void Encode(FBufferId BufferId, bool bForceKeyFrame, uint32 Bitrate = 0, TUniquePtr<AVEncoder::FEncoderVideoFrameCookie> Cookie = nullptr) = 0;
 
-	// We return a copy instead of a reference, to simplify multithread  access
+	/**
+	 * Return the current configuration
+	 */
 	virtual FVideoEncoderConfig GetConfig() const = 0;
 
 	// Allows setting bitrate and framerate at runtime.
 	virtual bool SetBitrate(uint32 Bitrate) = 0;
 	virtual bool SetFramerate(uint32 Framerate) = 0;
+
+	/**
+	 * Allows setting encoder parameters not exposed through FVideoEncoderConfig fields.
+	 * What parameters are valid are encoder dependent. See FVideoEncoderConfig::Options for examples
+	 */
 	virtual bool SetParameter(const FString& Parameter, const FString& Value) = 0;
 
+
+	/**
+	* Register/Unregister listeners.
+	* Most likely, you only need 1 listener per encoder, but supported several is needed due to some
+	* PixelStreaming peculiarities.
+	*/
 	virtual void RegisterListener(IVideoEncoderListener& Listener);
 	virtual void UnregisterListener(IVideoEncoderListener& Listener);
 
@@ -234,6 +464,9 @@ private:
 
 };
 
+/**
+ * Similar functionality to FVideoEncoder
+ */
 class AVENCODER_API FAudioEncoder
 {
 public:
@@ -241,6 +474,10 @@ public:
 	virtual const TCHAR* GetName() const = 0;
 	virtual const TCHAR* GetType() const = 0;
 	virtual bool Initialize(const FAudioEncoderConfig& Config) = 0;
+
+	/**
+	* Shutdown MUST be called before destruction
+	*/
 	virtual void Shutdown() = 0;
 
 	virtual void Encode(const FAudioFrame& Frame) = 0;
