Make function signature in header consistent with cpp
#rnx

diff --git a/Engine/Source/Runtime/AVEncoder/Public/VideoEncoderInput.h b/Engine/Source/Runtime/AVEncoder/Public/VideoEncoderInput.h
--- a/Engine/Source/Runtime/AVEncoder/Public/VideoEncoderInput.h
+++ b/Engine/Source/Runtime/AVEncoder/Public/VideoEncoderInput.h
@@ -213,7 +213,7 @@ namespace AVEncoder
 		using FReleaseVulkanTextureCallback = TFunction<void(VkImage_T*)>;
 
 #if PLATFORM_WINDOWS || PLATFORM_LINUX
-	void SetTexture(VkImage_T* InTexture, FReleaseVulkanTextureCallback InOnReleaseTexture);
+	void SetTexture(VkImage InTexture, FReleaseVulkanTextureCallback InOnReleaseTexture);
 #endif
 
 	protected:
