Fix compile error caused by changing signature to be consistent (reversed which direction it was made consistent)
#rnx

diff --git a/Engine/Source/Runtime/AVEncoder/Private/VideoEncoderInput.cpp b/Engine/Source/Runtime/AVEncoder/Private/VideoEncoderInput.cpp
--- a/Engine/Source/Runtime/AVEncoder/Private/VideoEncoderInput.cpp
+++ b/Engine/Source/Runtime/AVEncoder/Private/VideoEncoderInput.cpp
@@ -803,7 +803,7 @@ void FVideoEncoderInputFrame::SetTexture(CUarray InTexture, FReleaseCUDATextureC
 
 
 #if PLATFORM_WINDOWS || PLATFORM_LINUX
-void FVideoEncoderInputFrame::SetTexture(VkImage InTexture, FReleaseVulkanTextureCallback InOnReleaseTexture)
+void FVideoEncoderInputFrame::SetTexture(VkImage_T* InTexture, FReleaseVulkanTextureCallback InOnReleaseTexture)
 {
 	if (Format == EVideoFrameFormat::VULKAN_R8G8B8A8_UNORM)
 	{
diff --git a/Engine/Source/Runtime/AVEncoder/Public/VideoEncoderInput.h b/Engine/Source/Runtime/AVEncoder/Public/VideoEncoderInput.h
--- a/Engine/Source/Runtime/AVEncoder/Public/VideoEncoderInput.h
+++ b/Engine/Source/Runtime/AVEncoder/Public/VideoEncoderInput.h
@@ -213,7 +213,7 @@ namespace AVEncoder
 		using FReleaseVulkanTextureCallback = TFunction<void(VkImage_T*)>;
 
 #if PLATFORM_WINDOWS || PLATFORM_LINUX
-	void SetTexture(VkImage InTexture, FReleaseVulkanTextureCallback InOnReleaseTexture);
+	void SetTexture(VkImage_T* InTexture, FReleaseVulkanTextureCallback InOnReleaseTexture);
 #endif
 
 	protected:
