Ensure that NVAftermath defines are consistent between AVEncoder and VulkanRHI
#rnx

diff --git a/Engine/Source/Runtime/AVEncoder/AVEncoder.Build.cs b/Engine/Source/Runtime/AVEncoder/AVEncoder.Build.cs
--- a/Engine/Source/Runtime/AVEncoder/AVEncoder.Build.cs
+++ b/Engine/Source/Runtime/AVEncoder/AVEncoder.Build.cs
@@ -41,8 +41,9 @@ public class AVEncoder : ModuleRules
 			// ... add any modules that your module loads dynamically here ...
 		});
 
-		if (Target.Platform == UnrealTargetPlatform.Win64 || Target.Platform == UnrealTargetPlatform.Linux)
+		if (Target.IsInPlatformGroup(UnrealPlatformGroup.Windows) || Target.IsInPlatformGroup(UnrealPlatformGroup.Linux))
 		{
+			AddEngineThirdPartyPrivateStaticDependencies(Target, "NVAftermath");
 			PrivateIncludePathModuleNames.Add("VulkanRHI");
 		}
 
