Fix LinuxAArch64 compile error
#rnx

diff --git a/Engine/Source/Runtime/AVEncoder/AVEncoder.Build.cs b/Engine/Source/Runtime/AVEncoder/AVEncoder.Build.cs
--- a/Engine/Source/Runtime/AVEncoder/AVEncoder.Build.cs
+++ b/Engine/Source/Runtime/AVEncoder/AVEncoder.Build.cs
@@ -41,7 +41,7 @@ public class AVEncoder : ModuleRules
 			// ... add any modules that your module loads dynamically here ...
 		});
 
-		if (Target.IsInPlatformGroup(UnrealPlatformGroup.Windows) || Target.IsInPlatformGroup(UnrealPlatformGroup.Linux))
+		if (Target.IsInPlatformGroup(UnrealPlatformGroup.Windows) || Target.Platform == UnrealTargetPlatform.Linux)
 		{
 			AddEngineThirdPartyPrivateStaticDependencies(Target, "NVAftermath");
 			PrivateIncludePathModuleNames.Add("VulkanRHI");
