Fix merge error

diff --git a/Engine/Source/Runtime/AVEncoder/Public/IMFSampleWrapper.h b/Engine/Source/Runtime/AVEncoder/Public/IMFSampleWrapper.h
--- a/Engine/Source/Runtime/AVEncoder/Public/IMFSampleWrapper.h
+++ b/Engine/Source/Runtime/AVEncoder/Public/IMFSampleWrapper.h
@@ -5,13 +5,13 @@
 #include "MediaPacket.h"
 #include "Templates/RefCounting.h"
 
-//#if PLATFORM_WINDOWS
+#if PLATFORM_WINDOWS
 
-//#include "Windows/AllowWindowsPlatformTypes.h"
-//#include "Windows/PreWindowsApi.h"
-//	#include <mftransform.h>
-//#include "Windows/PostWindowsApi.h"
-//#include "Windows/HideWindowsPlatformTypes.h"
+// #include "Windows/AllowWindowsPlatformTypes.h"
+// #include "Windows/PreWindowsApi.h"
+// 	#include <mftransform.h>
+// #include "Windows/PostWindowsApi.h"
+// #include "Windows/HideWindowsPlatformTypes.h"
 
 namespace AVEncoder
 {
@@ -105,5 +105,5 @@ private:
 
 }
 
-#endif // PLATFORM_WINDOWS || (PLATFORM_XBOXONE && WITH_LEGACY_XDK)
+#endif // PLATFORM_WINDOWS
 
