Restore missing include

diff --git a/Engine/Source/Runtime/AVEncoder/Public/IMFSampleWrapper.h b/Engine/Source/Runtime/AVEncoder/Public/IMFSampleWrapper.h
--- a/Engine/Source/Runtime/AVEncoder/Public/IMFSampleWrapper.h
+++ b/Engine/Source/Runtime/AVEncoder/Public/IMFSampleWrapper.h
@@ -4,14 +4,15 @@
 
 #include "MediaPacket.h"
 #include "Templates/RefCounting.h"
+#include "VideoCommon.h"
 
 #if PLATFORM_WINDOWS
 
-// #include "Windows/AllowWindowsPlatformTypes.h"
-// #include "Windows/PreWindowsApi.h"
-// 	#include <mftransform.h>
-// #include "Windows/PostWindowsApi.h"
-// #include "Windows/HideWindowsPlatformTypes.h"
+/*#include "Windows/AllowWindowsPlatformTypes.h"
+#include "Windows/PreWindowsApi.h"
+	#include <mftransform.h>
+#include "Windows/PostWindowsApi.h"
+#include "Windows/HideWindowsPlatformTypes.h"*/
 
 namespace AVEncoder
 {
