Merge in PixelStreaming, AVEncoder, and SignallingServer changes from TensorWorks local repository

diff --git a/Engine/Plugins/Media/PixelStreaming/Source/PixelStreaming/Private/VideoCapturer.cpp b/Engine/Plugins/Media/PixelStreaming/Source/PixelStreaming/Private/VideoCapturer.cpp
--- a/Engine/Plugins/Media/PixelStreaming/Source/PixelStreaming/Private/VideoCapturer.cpp
+++ b/Engine/Plugins/Media/PixelStreaming/Source/PixelStreaming/Private/VideoCapturer.cpp
@@ -5,21 +5,22 @@
 
 #include "HAL/PlatformTime.h"
 #include "Misc/Timespan.h"
-
 #include "PixelStreamingFrameBuffer.h"
+#include "PixelStreamingSettings.h"
+#include "ClearQuad.h"
+#include "LatencyTester.h"
 
 #if PLATFORM_LINUX
 #include "CudaModule.h"
 #endif
 
-extern TAutoConsoleVariable<int32> CVarPixelStreamingEncoderUseBackBufferSize;
 extern TAutoConsoleVariable<FString> CVarPixelStreamingEncoderTargetSize;
 
 FVideoCapturer::FVideoCapturer()
 {
 	CurrentState = webrtc::MediaSourceInterface::SourceState::kInitializing;
 
-	LastTimestampUs = rtc::TimeMicros();
+	LastTimestampUs =  FTimespan::FromSeconds(FPlatformTime::Seconds()).GetTotalMicroseconds();
 
 	if (GDynamicRHI)
 	{
@@ -43,7 +44,7 @@ FVideoCapturer::FVideoCapturer()
 void FVideoCapturer::OnFrameReady(const FTexture2DRHIRef& FrameBuffer)
 {
 	FIntPoint Resolution = FrameBuffer->GetSizeXY();
-	int64 TimestampUs = rtc::TimeMicros();
+	int64 TimestampUs = FTimespan::FromSeconds(FPlatformTime::Seconds()).GetTotalMicroseconds();
 
 	int outWidth, outHeight, cropWidth, cropHeight, cropX, cropY;
 	if (!AdaptFrame(Resolution.X, Resolution.Y, TimestampUs, &outWidth, &outHeight, &cropWidth, &cropHeight, &cropX, &cropY))
@@ -54,12 +55,13 @@ void FVideoCapturer::OnFrameReady(const FTexture2DRHIRef& FrameBuffer)
 	if (CurrentState != webrtc::MediaSourceInterface::SourceState::kLive)
 		CurrentState = webrtc::MediaSourceInterface::SourceState::kLive;
 	
-	if (CVarPixelStreamingEncoderUseBackBufferSize.GetValueOnRenderThread() == 0)
+	// Set resolution of encoder using user-defined params (i.e. not the back buffer).
+	if (PixelStreamingSettings::CVarPixelStreamingUseBackBufferCaptureSize.GetValueOnRenderThread() == 0)
 	{
 		// set resolution based on cvars
-		FString EncoderTargetSize = CVarPixelStreamingEncoderTargetSize.GetValueOnRenderThread();
+		FString CaptureSize = PixelStreamingSettings::CVarPixelStreamingCaptureSize.GetValueOnRenderThread();
 		FString TargetWidth, TargetHeight;
-		bool bValidSize = EncoderTargetSize.Split(TEXT("x"), &TargetWidth, &TargetHeight);
+		bool bValidSize = CaptureSize.Split(TEXT("x"), &TargetWidth, &TargetHeight);
 		if (bValidSize)
 		{
 			Resolution.X = FCString::Atoi(*TargetWidth);
@@ -67,18 +69,22 @@ void FVideoCapturer::OnFrameReady(const FTexture2DRHIRef& FrameBuffer)
 		}
 		else
 		{
-			UE_LOG(PixelStreamer, Error, TEXT("CVarPixelStreamingEncoderTargetSize is not in a valid format: %s. It should be e.g: \"1920x1080\""), *EncoderTargetSize);
-			CVarPixelStreamingEncoderTargetSize->Set(*FString::Printf(TEXT("%dx%d"), Resolution.X, Resolution.Y));
+			UE_LOG(PixelStreamer, Error, TEXT("CVarPixelStreamingCaptureSize is not in a valid format: %s. It should be e.g: \"1920x1080\""), *CaptureSize);
+			PixelStreamingSettings::CVarPixelStreamingCaptureSize->Set(*FString::Printf(TEXT("%dx%d"), Resolution.X, Resolution.Y));
 		}
 
 		SetCaptureResolution(Resolution.X, Resolution.Y);
 	}
+	// Set the resolution based on the back buffer.
 	else
-		SetCaptureResolution(outWidth, outHeight); // set resolution based on back buffer
+	{
+		SetCaptureResolution(outWidth, outHeight); 
+	}
+		
 
 	AVEncoder::FVideoEncoderInputFrame* InputFrame = VideoEncoderInput->ObtainInputFrame();
-	InputFrame->PTS = FTimespan::FromSeconds(FPlatformTime::Seconds()).GetTicks();
-	auto FrameId = InputFrame->GetFrameID();
+	InputFrame->PTS = FTimespan::FromSeconds(FPlatformTime::Seconds()).GetTotalMicroseconds();
+	uint32 FrameId = InputFrame->GetFrameID();
 
 	if (!BackBuffers.Contains(InputFrame))
 	{
@@ -87,14 +93,14 @@ void FVideoCapturer::OnFrameReady(const FTexture2DRHIRef& FrameBuffer)
 		if (RHIName == TEXT("D3D11"))
 		{
 			FRHIResourceCreateInfo CreateInfo(TEXT("VideoCapturerBackBuffer"));
-			auto Texture = GDynamicRHI->RHICreateTexture2D(Width, Height, EPixelFormat::PF_B8G8R8A8, 1, 1, TexCreate_Shared | TexCreate_RenderTargetable | TexCreate_UAV, ERHIAccess::CopyDest, CreateInfo);
+			FTexture2DRHIRef Texture = GDynamicRHI->RHICreateTexture2D(Width, Height, EPixelFormat::PF_B8G8R8A8, 1, 1, TexCreate_Shared | TexCreate_RenderTargetable | TexCreate_UAV, ERHIAccess::CopyDest, CreateInfo);
 			InputFrame->SetTexture((ID3D11Texture2D*)Texture->GetNativeResource(), [&, InputFrame](ID3D11Texture2D* NativeTexture) { BackBuffers.Remove(InputFrame); });
 			BackBuffers.Add(InputFrame, Texture);
 		}
 		else if (RHIName == TEXT("D3D12"))
 		{
 			FRHIResourceCreateInfo CreateInfo(TEXT("VideoCapturerBackBuffer"));
-			auto Texture = GDynamicRHI->RHICreateTexture2D(Width, Height, EPixelFormat::PF_B8G8R8A8, 1, 1, TexCreate_Shared | TexCreate_RenderTargetable | TexCreate_UAV, ERHIAccess::CopyDest, CreateInfo);
+			FTexture2DRHIRef Texture = GDynamicRHI->RHICreateTexture2D(Width, Height, EPixelFormat::PF_B8G8R8A8, 1, 1, TexCreate_Shared | TexCreate_RenderTargetable | TexCreate_UAV, ERHIAccess::CopyDest, CreateInfo);
 			InputFrame->SetTexture((ID3D12Resource*)Texture->GetNativeResource(), [&, InputFrame](ID3D12Resource* NativeTexture) { BackBuffers.Remove(InputFrame); });
 			BackBuffers.Add(InputFrame, Texture);
 		}
@@ -108,11 +114,11 @@ void FVideoCapturer::OnFrameReady(const FTexture2DRHIRef& FrameBuffer)
 
 			// TODO (M84FIX) Replace with CUDA texture
 			// Create a texture that can be exposed to external memory
-			auto Texture = GDynamicRHI->RHICreateTexture2D(Width, Height, EPixelFormat::PF_B8G8R8A8, 1, 1, TexCreate_Shared | TexCreate_RenderTargetable | TexCreate_UAV, ERHIAccess::CopyDest, CreateInfo);
+			FTexture2DRHIRef Texture = GDynamicRHI->RHICreateTexture2D(Width, Height, EPixelFormat::PF_B8G8R8A8, 1, 1, TexCreate_Shared | TexCreate_RenderTargetable | TexCreate_UAV, ERHIAccess::CopyDest, CreateInfo);
 
-			auto VulkanTexture = static_cast<FVulkanTexture2D*>(Texture.GetReference());
+			FVulkanTexture2D* VulkanTexture = static_cast<FVulkanTexture2D*>(Texture.GetReference());
 
-			auto device = static_cast<FVulkanDynamicRHI*>(GDynamicRHI)->GetDevice()->GetInstanceHandle();
+			FVulkanDynamicRHI* device = static_cast<FVulkanDynamicRHI*>(GDynamicRHI)->GetDevice()->GetInstanceHandle();
 
 			// Get the CUarray to that textures memory making sure the clear it when done
 			int fd;
@@ -209,9 +215,29 @@ void FVideoCapturer::OnFrameReady(const FTexture2DRHIRef& FrameBuffer)
 		UE_LOG(PixelStreamer, Log, TEXT("%d backbuffers currently allocated"), BackBuffers.Num());
 	}
 
+	// Latency test pre capture
+	if(FLatencyTester::IsTestRunning() && FLatencyTester::GetTestStage() == FLatencyTester::ELatencyTestStage::PRE_CAPTURE)
+	{
+		FLatencyTester::RecordPreCaptureTime(FrameId);
+	}
+
+	// Actual texture copy (i.e the actual "capture")
 	CopyTexture(FrameBuffer, BackBuffers[InputFrame]);
 
-	rtc::scoped_refptr<FPixelStreamingFrameBuffer> Buffer = new rtc::RefCountedObject<FPixelStreamingFrameBuffer>(InputFrame, VideoEncoderInput);
+	// Latency test post capture
+	if(FLatencyTester::IsTestRunning() && FLatencyTester::GetTestStage() == FLatencyTester::ELatencyTestStage::POST_CAPTURE)
+	{
+		// Render a fully red frame for latency testing purposes
+		FRHICommandListImmediate& RHICmdList = FRHICommandListExecutor::GetImmediateCommandList();
+		FTexture2DRHIRef& DestinationTexture = BackBuffers[InputFrame];
+		FRHIRenderPassInfo RPInfo(DestinationTexture, ERenderTargetActions::Load_Store);
+		RHICmdList.BeginRenderPass(RPInfo, TEXT("ClearRT"));
+        DrawClearQuad(RHICmdList, FLinearColor::Red);
+        RHICmdList.EndRenderPass();
+		FLatencyTester::RecordPostCaptureTime(FrameId);
+	}
+
+	rtc::scoped_refptr<FPixelStreamingFrameBuffer> Buffer = new rtc::RefCountedObject<FPixelStreamingFrameBuffer>(BackBuffers[InputFrame], InputFrame, VideoEncoderInput);
 
 	webrtc::VideoFrame Frame = webrtc::VideoFrame::Builder().
 		set_video_frame_buffer(Buffer).
@@ -296,7 +322,7 @@ void FVideoCapturer::CopyTexture(const FTexture2DRHIRef& SourceTexture, FTexture
 bool FVideoCapturer::SetCaptureResolution(int NewCaptureWidth, int NewCaptureHeight)
 {
 	// Check is requested resolution is same as current resolution, if so, do nothing.
-	if(Width == NewCaptureWidth && Height == NewCaptureHeight)
+	if (Width == NewCaptureWidth && Height == NewCaptureHeight)
 		return false;
 
 	verifyf(NewCaptureWidth > 0, TEXT("Capture width must be greater than zero."));
diff --git a/Engine/Plugins/Media/PixelStreaming/Source/PixelStreaming/Private/WebRTCIncludes.h b/Engine/Plugins/Media/PixelStreaming/Source/PixelStreaming/Private/WebRTCIncludes.h
--- a/Engine/Plugins/Media/PixelStreaming/Source/PixelStreaming/Private/WebRTCIncludes.h
+++ b/Engine/Plugins/Media/PixelStreaming/Source/PixelStreaming/Private/WebRTCIncludes.h
@@ -57,6 +57,8 @@ THIRD_PARTY_INCLUDES_START
 #pragma warning(push)
 #pragma warning(disable: 4582 4583 6319 6323)
 
+#include "api/rtp_receiver_interface.h"
+#include "api/media_types.h"
 #include "api/media_stream_interface.h"
 #include "api/peer_connection_interface.h"
 #include "api/create_peerconnection_factory.h"
diff --git a/Engine/Plugins/Media/PixelStreaming/Source/PixelStreaming/Public/IPixelStreamingModule.h b/Engine/Plugins/Media/PixelStreaming/Source/PixelStreaming/Public/IPixelStreamingModule.h
--- a/Engine/Plugins/Media/PixelStreaming/Source/PixelStreaming/Public/IPixelStreamingModule.h
+++ b/Engine/Plugins/Media/PixelStreaming/Source/PixelStreaming/Public/IPixelStreamingModule.h
@@ -79,11 +79,5 @@ public:
 	 * Unfreeze Pixel Streaming.
 	 */
 	virtual void UnfreezeFrame() = 0;
-
-	// Player
-
-	virtual bool IsPlayerInitialized() const = 0;
-
-	virtual TSharedPtr<IMediaPlayer, ESPMode::ThreadSafe> CreatePlayer(IMediaEventSink& EventSink) = 0;
 };
 
diff --git a/Engine/Plugins/Media/PixelStreaming/Source/PixelStreaming/Public/PixelStreamingSettings.h b/Engine/Plugins/Media/PixelStreaming/Source/PixelStreaming/Public/PixelStreamingSettings.h
--- a/Engine/Plugins/Media/PixelStreaming/Source/PixelStreaming/Public/PixelStreamingSettings.h
+++ b/Engine/Plugins/Media/PixelStreaming/Source/PixelStreaming/Public/PixelStreamingSettings.h
@@ -5,9 +5,90 @@
 #include "CoreMinimal.h"
 #include "Engine/DeveloperSettings.h"
 #include "InputCoreTypes.h"
-
 #include "PixelStreamingSettings.generated.h"
 
+// Console variables (CVars)
+namespace PixelStreamingSettings
+{
+
+// Begin Encoder CVars
+	extern TAutoConsoleVariable<int32> CVarPixelStreamingEncoderTargetBitrate;
+	extern TAutoConsoleVariable<int32> CVarPixelStreamingEncoderMaxBitrate;
+	extern TAutoConsoleVariable<bool> CVarPixelStreamingDebugDumpFrame;
+	extern TAutoConsoleVariable<int32> CVarPixelStreamingEncoderMinQP;
+	extern TAutoConsoleVariable<int32> CVarPixelStreamingEncoderMaxQP;
+	extern TAutoConsoleVariable<FString> CVarPixelStreamingEncoderRateControl;
+	extern TAutoConsoleVariable<bool> CVarPixelStreamingEnableFillerData;
+	extern TAutoConsoleVariable<FString> CVarPixelStreamingEncoderMultipass;
+// End Encoder CVars
+
+// Begin Capturer CVars
+	extern TAutoConsoleVariable<int32> CVarPixelStreamingUseBackBufferCaptureSize;
+	extern TAutoConsoleVariable<FString> CVarPixelStreamingCaptureSize;
+// End Capturer CVars
+
+// Begin WebRTC CVars
+	extern TAutoConsoleVariable<FString> CVarPixelStreamingDegradationPreference;
+	extern TAutoConsoleVariable<int32> CVarPixelStreamingWebRTCMaxFps;
+	extern TAutoConsoleVariable<int32> CVarPixelStreamingWebRTCMinBitrate;
+	extern TAutoConsoleVariable<int32> CVarPixelStreamingWebRTCMaxBitrate;
+	extern TAutoConsoleVariable<int> CVarPixelStreamingWebRTCLowQpThreshold;
+	extern TAutoConsoleVariable<int> CVarPixelStreamingWebRTCHighQpThreshold;
+// End WebRTC CVars
+
+// Begin Pixel Streaming Plugin CVars
+	extern TAutoConsoleVariable<bool> CVarPixelStreamingHudStats;
+	extern TAutoConsoleVariable<int32> CVarFreezeFrameQuality;
+// Ends Pixel Streaming Plugin CVars
+
+// Begin utility functions etc.
+	AVEncoder::FVideoEncoder::RateControlMode GetRateControlCVar();
+	AVEncoder::FVideoEncoder::MultipassMode GetMultipassCVar();
+	webrtc::DegradationPreference GetDegradationPreference();
+// End utility functions etc.
+
+// Begin Command line args
+
+	inline bool IsAllowPixelStreamingCommands()
+	{
+		return FParse::Param(FCommandLine::Get(), TEXT("AllowPixelStreamingCommands"));
+	}
+
+	inline bool IsPixelStreamingHideCursor()
+	{
+		return FParse::Param(FCommandLine::Get(), TEXT("PixelStreamingHideCursor"));
+	}
+	
+	inline bool IsForceVP8()
+	{
+		return FParse::Param(FCommandLine::Get(), TEXT("PSForceVP8"));
+	}
+
+	inline bool GetSignallingServerIP(FString& OutSignallingServerIP)
+	{
+		return FParse::Value(FCommandLine::Get(), TEXT("PixelStreamingIP="), OutSignallingServerIP);
+	}
+
+	inline bool GetSignallingServerPort(uint16& OutSignallingServerPort)
+	{
+		return FParse::Value(FCommandLine::Get(), TEXT("PixelStreamingPort="), OutSignallingServerPort);
+	}
+
+	inline bool GetControlScheme(FString& OutControlScheme)
+	{
+		return FParse::Value(FCommandLine::Get(), TEXT("PixelStreamingControlScheme="), OutControlScheme);
+	}
+
+	inline bool GetFastPan(float& OutFastPan)
+	{
+		return FParse::Value(FCommandLine::Get(), TEXT("PixelStreamingFastPan="), OutFastPan);
+	}
+
+// End Command line args
+
+}
+
+// Config loaded/saved to an .ini file.
 UCLASS(config = PixelStreaming, defaultconfig, meta = (DisplayName = "PixelStreaming"))
 class PIXELSTREAMING_API UPixelStreamingSettings : public UDeveloperSettings
 {
diff --git a/Engine/Source/Programs/PixelStreaming/WebServers/SignallingWebServer/build-container.bat b/Engine/Source/Programs/PixelStreaming/WebServers/SignallingWebServer/build-container.bat
new file mode 100755
--- /dev/null
+++ b/Engine/Source/Programs/PixelStreaming/WebServers/SignallingWebServer/build-container.bat
@@ -0,0 +1 @@
+docker build -t cirrus:latest .
\ No newline at end of file
diff --git a/Engine/Source/Programs/PixelStreaming/WebServers/SignallingWebServer/build-container.sh b/Engine/Source/Programs/PixelStreaming/WebServers/SignallingWebServer/build-container.sh
new file mode 100644
--- /dev/null
+++ b/Engine/Source/Programs/PixelStreaming/WebServers/SignallingWebServer/build-container.sh
@@ -0,0 +1,2 @@
+#!/bin/sh
+sudo docker build -t cirrus:latest .
\ No newline at end of file
diff --git a/Engine/Source/Programs/PixelStreaming/WebServers/SignallingWebServer/cirrus.js b/Engine/Source/Programs/PixelStreaming/WebServers/SignallingWebServer/cirrus.js
--- a/Engine/Source/Programs/PixelStreaming/WebServers/SignallingWebServer/cirrus.js
+++ b/Engine/Source/Programs/PixelStreaming/WebServers/SignallingWebServer/cirrus.js
@@ -20,13 +20,13 @@ const defaultConfig = {
 	UseHTTPS: false,
 	UseAuthentication: false,
 	LogToFile: true,
-	HomepageFile: 'player.htm',
+	HomepageFile: 'player.html',
 	AdditionalRoutes: new Map(),
 	EnableWebserver: true
 };
 
 const argv = require('yargs').argv;
-var configFile = (typeof argv.configFile != 'undefined') ? argv.configFile.toString() : '.\\config.json';
+var configFile = (typeof argv.configFile != 'undefined') ? argv.configFile.toString() : './config.json';
 console.log(`configFile ${configFile}`);
 const config = require('./modules/config.js').init(configFile, defaultConfig);
 
@@ -174,7 +174,7 @@ sendGameSessionData();
 if(config.UseAuthentication){
 	if(config.EnableWebserver) {
 		app.get('/login', function(req, res){
-			res.sendFile(__dirname + '/login.htm');
+			res.sendFile(__dirname + '/login.html');
 		});
 	}
 
@@ -217,7 +217,7 @@ try {
 if(config.EnableWebserver) {
 	app.get('/', isAuthenticated('/login'), function (req, res) {
 		homepageFile = (typeof config.HomepageFile != 'undefined' && config.HomepageFile != '') ? config.HomepageFile.toString() : defaultConfig.HomepageFile;
-		homepageFilePath = path.join(__dirname, homepageFile)
+		homepageFilePath = path.join(__dirname, '/public', homepageFile)
 
 		fs.access(homepageFilePath, (err) => {
 			if (err) {
@@ -662,4 +662,4 @@ function sendPlayerDisconnectedToMatchmaker() {
 		type: 'clientDisconnected'
 	};
 	matchmaker.write(JSON.stringify(message));
-}
\ No newline at end of file
+}
diff --git a/Engine/Source/Programs/PixelStreaming/WebServers/SignallingWebServer/config.json b/Engine/Source/Programs/PixelStreaming/WebServers/SignallingWebServer/config.json
new file mode 100644
--- /dev/null
+++ b/Engine/Source/Programs/PixelStreaming/WebServers/SignallingWebServer/config.json
@@ -0,0 +1,10 @@
+{
+	"UseFrontend": false,
+	"UseMatchmaker": false,
+	"UseHTTPS": false,
+	"UseAuthentication": false,
+	"LogToFile": true,
+	"HomepageFile": "player.html",
+	"AdditionalRoutes": {},
+	"EnableWebserver": true
+}
\ No newline at end of file
diff --git a/Engine/Source/Programs/PixelStreaming/WebServers/SignallingWebServer/public/login.html b/Engine/Source/Programs/PixelStreaming/WebServers/SignallingWebServer/public/login.html
new file mode 100644
--- /dev/null
+++ b/Engine/Source/Programs/PixelStreaming/WebServers/SignallingWebServer/public/login.html
@@ -0,0 +1,29 @@
+<!-- Copyright Epic Games, Inc. All Rights Reserved. -->
+<!DOCTYPE html>
+<html>
+
+<head>
+    <meta charset="utf-8">
+    <link rel="shortcut icon" href="/images/favicon.ico" type="image/x-icon">
+    <link rel="icon" type="image/png" sizes="96x96" href="/images/favicon-96x96.png">
+    <link rel="icon" type="image/png" sizes="32x32" href="/images/favicon-32x32.png">
+    <link rel="icon" type="image/png" sizes="16x16" href="/images/favicon-16x16.png">
+    <link type="text/css" rel="stylesheet" href="login.css">
+    <title>Login</title>
+</head>
+
+<body>
+    <form action="/login" method="post">
+        <div class="entry">
+          <label for="name">Username:</label> <input type="text" id="username" name="username" autocomplete="username">
+        </div>
+        <div class="entry">
+          <label for="mail">Password:</label> <input type="password" id="password" name="password" autocomplete="current-password">
+        </div>
+        <div class="entry button">
+          <button type="submit">Login</button>
+        </div>
+    </form>
+</body>
+
+</html>
\ No newline at end of file
diff --git a/Engine/Source/Programs/PixelStreaming/WebServers/SignallingWebServer/public/player.css b/Engine/Source/Programs/PixelStreaming/WebServers/SignallingWebServer/public/player.css
--- a/Engine/Source/Programs/PixelStreaming/WebServers/SignallingWebServer/public/player.css
+++ b/Engine/Source/Programs/PixelStreaming/WebServers/SignallingWebServer/public/player.css
@@ -25,13 +25,13 @@ body{
 	z-index: 10;
 }
 
-#statsContainer{
+.statsContainer {
 	background-color: rgba(0,0,0,0.8);
 	text-align: left;
 	display: block;
 }
 
-#stats{
+.stats {
 	font-size: 14px;
 	font-weight: bold;
 	padding: 6px;
diff --git a/Engine/Source/Programs/PixelStreaming/WebServers/SignallingWebServer/public/player.html b/Engine/Source/Programs/PixelStreaming/WebServers/SignallingWebServer/public/player.html
new file mode 100644
--- /dev/null
+++ b/Engine/Source/Programs/PixelStreaming/WebServers/SignallingWebServer/public/player.html
@@ -0,0 +1,111 @@
+<!-- Copyright 1998-2018 Epic Games, Inc. All Rights Reserved. -->
+<!DOCTYPE HTML>
+<html>
+<head>
+	<link rel="shortcut icon" href="/images/favicon.ico" type="image/x-icon">
+	<link rel="icon" type="image/png" sizes="96x96" href="/images/favicon-96x96.png">
+	<link rel="icon" type="image/png" sizes="32x32" href="/images/favicon-32x32.png">
+	<link rel="icon" type="image/png" sizes="16x16" href="/images/favicon-16x16.png">
+	<link type="text/css" rel="stylesheet" href="player.css">
+    <script type="text/javascript" src="https://webrtc.github.io/adapter/adapter-latest.js"></script>
+    <script type="text/javascript" src="scripts/webRtcPlayer.js"></script>
+    <script type="text/javascript" src="scripts/app.js"></script>
+</head>
+
+<body onload="load()">
+	<div id="playerUI">
+		<div id="player"></div>
+		<div id="overlay" class="overlay text-light bg-dark">
+			<div>
+				<div id="qualityStatus" class="greyStatus">&#9679</div>
+				<div id="overlayButton">+</div>
+			</div>
+			<div id="overlaySettings">
+				<div id="KickOthers">
+					<div class="settings-text">Kick all other players</div>
+					<label class="btn-overlay">
+						<input type="button" id="kick-other-players-button" class="overlay-button btn-flat" value="Kick">
+					</label>
+				</div>
+				<div id="FillWindow">
+					<div class="settings-text">Enlarge Display to Fill Window</div>
+					<label class="tgl-switch">
+						<input type="checkbox" id="enlarge-display-to-fill-window-tgl" class="tgl tgl-flat" checked>
+						<div class="tgl-slider"></div>
+					</label>
+				</div>
+				<div id="QualityControlOwnership">
+					<div class="settings-text">Quality control ownership</div>
+					<label class="tgl-switch">
+						<input type="checkbox" id="quality-control-ownership-tgl" class="tgl tgl-flat">
+						<div class="tgl-slider"></div>
+					</label>
+				</div>
+				<div id="PrioritiseQuality">
+					<div class="settings-text">Prioritise Quality</div>
+					<label class="tgl-switch">
+						<input type="checkbox" id="prioritise-quality-tgl" class="tgl tgl-flat" data-toggle="collapse" href="#qualityParamsContainer">
+						<div class="tgl-slider"></div>
+					</label>
+					<div id="qualityParamsContainer" class="collapse">
+						<div class="form-group">
+							<label for="low-bitrate-text">Low Bitrate (kbps)</label>
+							<input type="number" class="form-control" id="low-bitrate-text" value="2000" min="0" max="100000" />
+							<label for="high-bitrate-text">High Bitrate (kbps)</label>
+							<input type="number" class="form-control" id="high-bitrate-text" value="10000" min="0" max="100000" />
+							<label for="min-fps-text">Min FPS</label>
+							<input type="number" class="form-control" id="min-fps-text" value="10" min="1" max="120" />
+							<input id="quality-params-submit" class="btn btn-primary btn-lg mt-3" type="button" value="Apply">
+						</div>
+					</div>
+				</div>
+				<div id="RateControlContainer">
+					<div class="settings-text">Rate Control</div>
+					<label class="tgl-switch">
+						<select id="rate-control-dropdown">
+							<option value="CBR" selected>CBR</option>
+							<option value="VBR">VBR</option>
+							<option value="ConstQP">ConstQP</option>
+						</select>
+					</label>
+				</div>
+				<div id="MinQPContainer">
+					<div class="settings-text">Qual control (0=max qual/br)</div>
+					<label class="tgl-switch">
+						<input type="number" id="minqp-number" value="21" min="0" max="51">
+					</label>
+				</div>
+				<div id="ShowFPS">
+					<div class="settings-text">Show FPS</div>
+					<label class="btn-overlay">
+						<input type="button" id="show-fps-button" class="overlay-button btn-flat" value="Toggle">
+					</label>
+				</div>
+				<div id="MatchViewportResolution">
+					<div class="settings-text">Match Viewport Resolution</div>
+					<label class="tgl-switch">
+						<input type="checkbox" id="match-viewport-res-tgl" class="tgl tgl-flat">
+						<div class="tgl-slider"></div>
+					</label>
+				</div>
+				<div id="Stats">
+					<div class="settings-text">Show Stats</div>
+					<label class="tgl-switch">
+						<input type="checkbox" id="show-stats-tgl" class="tgl tgl-flat" checked>
+						<div class="tgl-slider"></div>
+					</label>
+					<div id="statsContainer" class="statsContainer">
+						<div id="stats" class="stats"></div>
+					</div>
+				</div>
+				<div id="LatencyTest">
+					<Button onclick="sendStartLatencyTest()">Test Latency</Button>
+					<div id="LatencyStatsContainer" class="statsContainer">
+						<div id=LatencyStats class="stats">No stats yet...</div>
+					</div>
+				</div>
+			</div>
+		</div>
+	</div>
+</body>
+</html>
diff --git a/Engine/Source/Programs/PixelStreaming/WebServers/SignallingWebServer/run-container.bat b/Engine/Source/Programs/PixelStreaming/WebServers/SignallingWebServer/run-container.bat
new file mode 100755
--- /dev/null
+++ b/Engine/Source/Programs/PixelStreaming/WebServers/SignallingWebServer/run-container.bat
@@ -0,0 +1 @@
+docker run -it --rm --init cirrus:latest
\ No newline at end of file
diff --git a/Engine/Source/Programs/PixelStreaming/WebServers/SignallingWebServer/run-container.sh b/Engine/Source/Programs/PixelStreaming/WebServers/SignallingWebServer/run-container.sh
new file mode 100644
--- /dev/null
+++ b/Engine/Source/Programs/PixelStreaming/WebServers/SignallingWebServer/run-container.sh
@@ -0,0 +1,2 @@
+#!/bin/sh
+sudo docker run -it --rm --init cirrus:latest
\ No newline at end of file
diff --git a/Engine/Source/Programs/PixelStreaming/WebServers/SignallingWebServer/run.sh b/Engine/Source/Programs/PixelStreaming/WebServers/SignallingWebServer/run.sh
new file mode 100644
--- /dev/null
+++ b/Engine/Source/Programs/PixelStreaming/WebServers/SignallingWebServer/run.sh
@@ -0,0 +1,3 @@
+#!/bin/bash
+bash setup.sh
+sudo node cirrus.js
\ No newline at end of file
diff --git a/Engine/Source/Programs/PixelStreaming/WebServers/SignallingWebServer/scripts/app.js b/Engine/Source/Programs/PixelStreaming/WebServers/SignallingWebServer/scripts/app.js
--- a/Engine/Source/Programs/PixelStreaming/WebServers/SignallingWebServer/scripts/app.js
+++ b/Engine/Source/Programs/PixelStreaming/WebServers/SignallingWebServer/scripts/app.js
@@ -130,6 +130,25 @@ function setupHtmlEvents() {
 		};
 	}
 
+	// Rate control drop down selection
+	let rateControlSelect = document.getElementById('rate-control-dropdown')
+	if (rateControlSelect !== null) {
+		rateControlSelect.onchange = function (event) {
+			var selectedValue = event.target.value;  
+			sendRateControlConsoleCommand(selectedValue);
+		}
+	}
+
+	// Minimum quality prioritisation spinner
+	let minqpSpinner = document.getElementById('minqp-number')
+	if (minqpSpinner !== null) {
+		minqpSpinner.onchange = function (event) {
+			var selectedValue = event.target.value;  
+			sendMinQPConsoleCommand(selectedValue);
+		}
+	}
+	
+
 	let showFPSButton = document.getElementById('show-fps-button');
 	if (showFPSButton !== null) {
 		showFPSButton.onclick = function (event) {
@@ -164,9 +183,47 @@ function setupHtmlEvents() {
 	}
 }
 
+function sendStartLatencyTest() {
+	// We need WebRTC to be active to do a latency test.
+	if(!webRtcPlayerObj)
+	{
+		return;
+	}
+
+	let onTestStarted = function(StartTimeMs)
+	{
+		let descriptor = {
+			StartTime: StartTimeMs
+		};
+		emitDescriptor(MessageType.LatencyTest, descriptor);
+	};
+
+	webRtcPlayerObj.startLatencyTest(onTestStarted);
+}
+
+function sendMinQPConsoleCommand(minQP) {
+	if(minQP) {
+		let command = 'PixelStreaming.Encoder.MinQP ' + minQP;
+		let consoleDescriptor = {
+			Console: command
+		};
+		emitUIInteraction(consoleDescriptor);
+	}
+}
+
+function sendRateControlConsoleCommand(rateControlMode) {
+	if(rateControlMode) {
+		let command = 'PixelStreaming.Encoder.RateControl ' + rateControlMode;
+		let consoleDescriptor = {
+			Console: command
+		};
+		emitUIInteraction(consoleDescriptor);
+	}
+}
+
 function sendQualityConsoleCommands(descriptor) {
 	if (descriptor.PrioritiseQuality !== null) {
-		let command = 'Streamer.PrioritiseQuality ' + descriptor.PrioritiseQuality;
+		let command = 'PixelStreaming.Encoder.PrioritizeQuality ' + descriptor.PrioritiseQuality;
 		let consoleDescriptor = {
 			Console: command
 		};
@@ -174,7 +231,7 @@ function sendQualityConsoleCommands(descriptor) {
 	}
 
 	if (descriptor.LowBitrate !== null) {
-		let command = 'Streamer.LowBitrate ' + descriptor.LowBitrate;
+		let command = 'PixelStreaming.Encoder.LowBitrate ' + descriptor.LowBitrate;
 		let consoleDescriptor = {
 			Console: command
 		};
@@ -182,7 +239,7 @@ function sendQualityConsoleCommands(descriptor) {
 	}
 
 	if (descriptor.HighBitrate !== null) {
-		let command = 'Streamer.HighBitrate ' + descriptor.HighBitrate;
+		let command = 'PixelStreaming.Encoder.HighBitrate ' + descriptor.HighBitrate;
 		let consoleDescriptor = {
 			Console: command
 		};
@@ -190,7 +247,7 @@ function sendQualityConsoleCommands(descriptor) {
 	}
 
 	if (descriptor.MinFPS !== null) {
-		var command = 'Streamer.MinFPS ' + descriptor.MinFPS;
+		var command = 'PixelStreaming.Encoder.MinFPS ' + descriptor.MinFPS;
 		let consoleDescriptor = {
 			Console: command
 		};
@@ -355,14 +412,15 @@ function removeResponseEventListener(name) {
 	responseEventListeners.remove(name);
 }
 
-// Must be kept in sync with PixelStreamingProtocol::EToClientMsg C++ enum.
+// Must be kept in sync with PixelStreamingProtocol::EToPlayerMsg C++ enum.
 const ToClientMessageType = {
 	QualityControlOwnership: 0,
 	Response: 1,
 	Command: 2,
 	FreezeFrame: 3,
 	UnfreezeFrame: 4,
-	VideoEncoderAvgQP: 5
+	VideoEncoderAvgQP: 5,
+	LatencyTest: 6
 };
 
 var VideoEncoderQP = "N/A";
@@ -443,6 +501,7 @@ function setupWebRtcPlayer(htmlElement, config) {
 			}
 		} else if (view[0] === ToClientMessageType.QualityControlOwnership) {
 			let ownership = view[1] === 0 ? false : true;
+			console.log("Received quality controller message, will control quality: " + ownership);
 			// If we own the quality control, we can't relenquish it. We only loose
 			// quality control when another peer asks for it
 			if (qualityControlOwnershipCheckBox !== null) {
@@ -476,6 +535,15 @@ function setupWebRtcPlayer(htmlElement, config) {
 		} else if (view[0] === ToClientMessageType.VideoEncoderAvgQP) {
 			VideoEncoderQP = new TextDecoder("utf-16").decode(data.slice(1));
 			console.log(`received VideoEncoderAvgQP ${VideoEncoderQP}`);
+		} else if(view[0] == ToClientMessageType.LatencyTest){
+			let latencyTimingsAsString = new TextDecoder("utf-16").decode(data.slice(1));
+			console.log("Got latency timings from UE.")
+			console.log(latencyTimingsAsString);
+			let latencyTimingsFromUE = JSON.parse(latencyTimingsAsString);
+			if(webRtcPlayerObj)
+			{
+				webRtcPlayerObj.latencyTestTimings.SetUETimings(latencyTimingsFromUE);
+			}
 		} else {
 			console.error(`unrecognized data received, packet ID ${view[0]}`);
 		}
@@ -579,7 +647,8 @@ function onWebRtcAnswer(webRTCData) {
 		statsText += `<div style="color: ${color}">Bitrate (kbps): ${aggregatedStats.hasOwnProperty('bitrate') ? numberFormat.format(aggregatedStats.bitrate) : 'N/A'}</div>`;
 		statsText += `<div>Framerate: ${aggregatedStats.hasOwnProperty('framerate') ? numberFormat.format(aggregatedStats.framerate) : 'N/A'}</div>`;
 		statsText += `<div>Frames dropped: ${aggregatedStats.hasOwnProperty('framesDropped') ? numberFormat.format(aggregatedStats.framesDropped) : 'N/A'}</div>`;
-		statsText += `<div>Latency (ms): ${aggregatedStats.hasOwnProperty('currentRoundTripTime') ? numberFormat.format(aggregatedStats.currentRoundTripTime * 1000) : 'N/A'}</div>`;
+		statsText += `<div>Net RTT (ms): ${aggregatedStats.hasOwnProperty('currentRoundTripTime') ? numberFormat.format(aggregatedStats.currentRoundTripTime * 1000) : 'N/A'}</div>`;
+		statsText += `<div>Browser receive to composite (ms): ${aggregatedStats.hasOwnProperty('receiveToCompositeMs') ? numberFormat.format(aggregatedStats.receiveToCompositeMs) : 'N/A'}</div>`;
 		statsText += `<div style="color: ${color}">Video Quantization Parameter: ${VideoEncoderQP}</div>`;
 
 		let statsDiv = document.getElementById("stats");
@@ -600,8 +669,28 @@ function onWebRtcAnswer(webRTCData) {
 
 	webRtcPlayerObj.aggregateStats(1 * 1000 /*Check every 1 second*/);
 
-	//let displayStats = () => { webRtcPlayerObj.getStats( (s) => { s.forEach(stat => { console.log(JSON.stringify(stat)); }); } ); }
-	//var displayStatsIntervalId = setInterval(displayStats, 30 * 1000);
+	webRtcPlayerObj.latencyTestTimings.OnAllLatencyTimingsReady = function(timings)
+	{
+		let endToEndLatency = timings.FrameDisplayTimeMs - timings.TestStartTimeMs;
+		let browserToUELatency = timings.ReceiptTimeMs - timings.TestStartTimeMs;
+		let ueToDisplayLatency = timings.FrameDisplayTimeMs - timings.PreCaptureTimeMs;
+		let uePreCapture = timings.PreCaptureTimeMs - timings.ReceiptTimeMs;
+		let uePixelStreamLatency = timings.PostEncodeTimeMs - timings.PreCaptureTimeMs;
+		let captureLatency = timings.PostCaptureTimeMs - timings.PreCaptureTimeMs;
+		let encodeLatency = timings.PostEncodeTimeMs - timings.PreEncodeTimeMs;
+		let browserSideLatency = endToEndLatency - uePreCapture - uePixelStreamLatency - (browserToUELatency * 2.0);
+		
+		let latencyStatsInnerHTML = '';
+		latencyStatsInnerHTML += `<div>Total stream latency (ms): ${endToEndLatency}</div>`;
+		latencyStatsInnerHTML += `<div>Browser to UE network (ms): ${browserToUELatency}</div>`;
+		latencyStatsInnerHTML += `<div>UE Capture to Browser Display (ms): ${ueToDisplayLatency}</div>`;
+		latencyStatsInnerHTML += `<div>UE Pre-capture (ms): ${uePreCapture}</div>`;
+		latencyStatsInnerHTML += `<div>UE Capture+Encode (ms): ${uePixelStreamLatency}</div>`;
+		latencyStatsInnerHTML += `<div>UE Capture (ms): ${captureLatency}</div>`;
+		latencyStatsInnerHTML += `<div>UE Encode (ms): ${encodeLatency}</div>`;
+		latencyStatsInnerHTML += `<div>Browser receive to display (ms): ${browserSideLatency}</div>`;
+		document.getElementById("LatencyStats").innerHTML = latencyStatsInnerHTML;		
+	}
 }
 
 function onWebRtcIce(iceCandidate) {
@@ -757,8 +846,11 @@ function resizePlayerStyle(event) {
 
 	updateVideoStreamSize();
 
-	if (playerElement.classList.contains('fixed-size'))
+	if (playerElement.classList.contains('fixed-size')){
+		setupMouseAndFreezeFrame(playerElement)
 		return;
+	}
+		
 
 	let checkBox = document.getElementById('enlarge-display-to-fill-window-tgl');
 	let windowSmallerThanPlayer = window.innerWidth < playerElement.videoWidth || window.innerHeight < playerElement.videoHeight;
@@ -772,6 +864,10 @@ function resizePlayerStyle(event) {
 		resizePlayerStyleToArbitrarySize(playerElement);
 	}
 	
+	setupMouseAndFreezeFrame(playerElement)
+}
+
+function setupMouseAndFreezeFrame(playerElement) {
 	// Calculating and normalizing positions depends on the width and height of
 	// the player.
 	playerElementClientRect = playerElement.getBoundingClientRect();
@@ -828,6 +924,7 @@ const MessageType = {
 	AverageBitrateRequest: 3,
 	StartStreaming: 4,
 	StopStreaming: 5,
+	LatencyTest: 6,
 
 	/**********************************************************************/
 
@@ -1519,7 +1616,10 @@ function onExpandOverlay_Click(/* e */) {
 function start() {
 	// update "quality status" to "disconnected" state
 	let qualityStatus = document.getElementById("qualityStatus");
-	qualityStatus.className = "grey-status";
+	if (qualityStatus) {
+		qualityStatus.className = "grey-status";
+	}
+	
 
 	let statsDiv = document.getElementById("stats");
 	if (statsDiv) {
diff --git a/Engine/Source/Programs/PixelStreaming/WebServers/SignallingWebServer/scripts/webRtcPlayer.js b/Engine/Source/Programs/PixelStreaming/WebServers/SignallingWebServer/scripts/webRtcPlayer.js
--- a/Engine/Source/Programs/PixelStreaming/WebServers/SignallingWebServer/scripts/webRtcPlayer.js
+++ b/Engine/Source/Programs/PixelStreaming/WebServers/SignallingWebServer/scripts/webRtcPlayer.js
@@ -23,23 +23,86 @@
 
         //**********************
         //Config setup
-        //**********************;
+        //**********************
 		this.cfg = parOptions.peerConnectionOptions || {};
 		this.cfg.sdpSemantics = 'unified-plan';
+        // this.cfg.rtcAudioJitterBufferMaxPackets = 10;
+        // this.cfg.rtcAudioJitterBufferFastAccelerate = true;
+        // this.cfg.rtcAudioJitterBufferMinDelayMs = 0;
+
+
 		//If this is true in Chrome 89+ SDP is sent that is incompatible with UE WebRTC and breaks.
-        this.cfg.offerExtmapAllowMixed = false;
+        //this.cfg.offerExtmapAllowMixed = false;
         this.pcClient = null;
         this.dcClient = null;
         this.tnClient = null;
 
         this.sdpConstraints = {
-          offerToReceiveAudio: 1,
+          offerToReceiveAudio: 1, //Note: if you don't need audio you can get improved latency by turning this off.
           offerToReceiveVideo: 1
         };
 
         // See https://www.w3.org/TR/webrtc/#dom-rtcdatachannelinit for values
         this.dataChannelOptions = {ordered: true};
 
+        // ToDo: get this useMic from url string like ?useMic or similar
+        //this.useMic = false;
+
+        //**********************
+        //Variables
+        //**********************
+        this.latencyTestTimings = 
+        {
+            TestStartTimeMs: null,
+            ReceiptTimeMs: null,
+            PreCaptureTimeMs: null,
+            PostCaptureTimeMs: null,
+            PreEncodeTimeMs: null,
+            PostEncodeTimeMs: null,
+            FrameDisplayTimeMs: null,
+            Reset: function()
+            {
+                this.TestStartTimeMs = null;
+                this.ReceiptTimeMs = null;
+                this.PreCaptureTimeMs = null;
+                this.PostCaptureTimeMs = null;
+                this.PreEncodeTimeMs = null;
+                this.PostEncodeTimeMs = null;
+                this.FrameDisplayTimeMs = null;
+            },
+            HasAllTimings: function()
+            {
+                return this.TestStartTimeMs && 
+                this.ReceiptTimeMs && 
+                this.PreCaptureTimeMs && 
+                this.PostCaptureTimeMs && 
+                this.PreEncodeTimeMs && 
+                this.PostEncodeTimeMs && 
+                this.FrameDisplayTimeMs;
+            },
+            SetUETimings: function(UETimings)
+            {
+                this.ReceiptTimeMs = UETimings.ReceiptTimeMs;
+                this.PreCaptureTimeMs = UETimings.PreCaptureTimeMs;
+                this.PostCaptureTimeMs = UETimings.PostCaptureTimeMs;
+                this.PreEncodeTimeMs = UETimings.PreEncodeTimeMs;
+                this.PostEncodeTimeMs = UETimings.PostEncodeTimeMs;
+                if(this.HasAllTimings())
+                {
+                    this.OnAllLatencyTimingsReady(this);
+                }
+            },
+            SetFrameDisplayTime: function(TimeMs)
+            {
+                this.FrameDisplayTimeMs = TimeMs;
+                if(this.HasAllTimings())
+                {
+                    this.OnAllLatencyTimingsReady(this);
+                }
+            },
+            OnAllLatencyTimingsReady: function(Timings){}
+        }
+
         //**********************
         //Functions
         //**********************
@@ -56,6 +119,28 @@
                     self.onVideoInitialised();
                 }
             }, true);
+			
+			// Check if request video frame callback is supported
+			if ('requestVideoFrameCallback' in HTMLVideoElement.prototype) {
+				// The API is supported! 
+				
+				const onVideoFrameReady = (now, metadata) => {
+					
+					if(metadata.receiveTime && metadata.expectedDisplayTime)
+					{
+						const receiveToCompositeMs = metadata.presentationTime - metadata.receiveTime;
+						self.aggregatedStats.receiveToCompositeMs = receiveToCompositeMs;
+					}
+					
+				  
+					// Re-register the callback to be notified about the next frame.
+					video.requestVideoFrameCallback(onVideoFrameReady);
+				};
+				
+				// Initially register the callback to be notified about the first frame.
+				video.requestVideoFrameCallback(onVideoFrameReady);
+			}
+			
             return video;
         }
 
@@ -91,12 +176,16 @@
 				console.log('Set video stream from ontrack');
 				
             }
+			
         };
 
         setupDataChannel = function(pc, label, options) {
             try {
-                var datachannel = pc.createDataChannel(label, options)
+                let datachannel = pc.createDataChannel(label, options);
                 console.log(`Created datachannel (${label})`)
+
+                // Inform browser we would like binary data as an ArrayBuffer (FF chooses Blob by default!)
+                datachannel.binaryType = "arraybuffer";
                 
                 datachannel.onopen = function (e) {
                   console.log(`data channel (${label}) connect`)
@@ -136,7 +225,8 @@
             		// (andriy): increase start bitrate from 300 kbps to 20 mbps and max bitrate from 2.5 mbps to 100 mbps
                     // (100 mbps means we don't restrict encoder at all)
                     // after we `setLocalDescription` because other browsers are not c happy to see google-specific config
-            		offer.sdp = offer.sdp.replace(/(a=fmtp:\d+ .*level-asymmetry-allowed=.*)\r\n/gm, "$1;x-google-start-bitrate=10000;x-google-max-bitrate=20000\r\n");
+            		offer.sdp = offer.sdp.replace(/(a=fmtp:\d+ .*level-asymmetry-allowed=.*)\r\n/gm, "$1;x-google-start-bitrate=10000;x-google-max-bitrate=100000\r\n");
+                    //offer.sdp = offer.sdp.replace("http://www.webrtc.org/experiments/rtp-hdrext/playout-delay", "http://www.webrtc.org/experiments/rtp-hdrext/playout-delay 0 0\r\na=name:playout-delay");
             		self.onWebRtcOffer(offer);
                 }
             },
@@ -225,7 +315,12 @@
                     }
                 });
 
-                //console.log(JSON.stringify(newStat));
+				
+				if(self.aggregatedStats.receiveToCompositeMs)
+				{
+					newStat.receiveToCompositeMs = self.aggregatedStats.receiveToCompositeMs;
+				}
+				
                 self.aggregatedStats = newStat;
 
                 if(self.onAggregatedStats)
@@ -233,10 +328,72 @@
             }
         };
 
+        setupTracksToSendAsync = async function(pc){
+            const stream = await navigator.mediaDevices.getUserMedia({video: false, audio: false});
+            if(stream)
+            {
+                for (const track of stream.getTracks()) {
+                    if(track.kind && track.kind == "audio")
+                    {
+                        pc.addTrack(track);
+                    }
+                }
+            }
+        };
+
+
         //**********************
         //Public functions
         //**********************
 
+        this.startLatencyTest = function(onTestStarted) {
+            // Can't start latency test without a video element
+            if(!self.video)
+            {
+                return;
+            }
+            let videoCanvas = document.createElement("canvas");
+            videoCanvas.style.display = 'none';
+            var ctx = videoCanvas.getContext('2d');
+            videoCanvas.width = self.video.videoWidth;
+            videoCanvas.height = self.video.videoHeight;
+
+            self.latencyTestTimings.Reset();
+            self.video.focus();
+
+            let checkCanvasSpecialLatencyPixels = function()
+            {
+                // Once we have our canvas pixel checker running we consider the test started properly
+                if(self.latencyTestTimings.TestStartTimeMs == null)
+                {
+                    self.latencyTestTimings.TestStartTimeMs = Date.now();
+                    onTestStarted(self.latencyTestTimings.TestStartTimeMs);
+                }
+
+                // draw the video to the canvas so we can analyse pixels in the canvas
+                ctx.drawImage(self.video, 0,0);
+                let middlePixelW = self.video.videoWidth * 0.5
+                let middlePixelH = self.video.videoHeight * 0.5
+                let data = ctx.getImageData(middlePixelW, middlePixelH, 1, 1).data;
+                let rgb = [ data[0], data[1], data[2] ];
+                let redValue = rgb[0];
+                if(redValue == 255)
+                {
+                    delete videoCanvas;
+                    console.log("Got special latency frame!");
+                    self.latencyTestTimings.SetFrameDisplayTime(Date.now());
+                }
+                else
+                {
+                    window.requestAnimationFrame(checkCanvasSpecialLatencyPixels);
+                }
+            }
+
+            // start checking for special latency pixels
+            window.requestAnimationFrame(checkCanvasSpecialLatencyPixels);
+            
+        }
+
         //This is called when revceiving new ice candidates individually instead of part of the offer
         //This is currently not used but would be called externally from this class
         this.handleCandidateFromServer = function(iceCandidate) {
@@ -255,9 +412,19 @@
                 self.pcClient = null;
             }
             self.pcClient = new RTCPeerConnection(self.cfg);
+
             setupPeerConnection(self.pcClient);
             self.dcClient = setupDataChannel(self.pcClient, 'cirrus', self.dataChannelOptions);
             handleCreateOffer(self.pcClient);
+
+            //if we were going to support browser sending mic/video into UE
+            // setupTracksToSendAsync(self.pcClient).finally(function()
+            // {
+            //     setupPeerConnection(self.pcClient);
+            //     self.dcClient = setupDataChannel(self.pcClient, 'cirrus', self.dataChannelOptions);
+            //     handleCreateOffer(self.pcClient);
+            // });
+            
         };
 
         //Called externaly when an answer is received from the server
@@ -265,6 +432,12 @@
             console.log(`Received answer:\n${answer}`);
             var answerDesc = new RTCSessionDescription(answer);
             self.pcClient.setRemoteDescription(answerDesc);
+
+            let receivers = self.pcClient.getReceivers();
+            for(let receiver of receivers)
+            {
+                receiver.playoutDelayHint = 0;
+            }
         };
 
         this.close = function(){
diff --git a/Engine/Source/Programs/PixelStreaming/WebServers/SignallingWebServer/setup.sh b/Engine/Source/Programs/PixelStreaming/WebServers/SignallingWebServer/setup.sh
new file mode 100644
--- /dev/null
+++ b/Engine/Source/Programs/PixelStreaming/WebServers/SignallingWebServer/setup.sh
@@ -0,0 +1 @@
+npm install
\ No newline at end of file
diff --git a/Engine/Source/Runtime/AVEncoder/AVEncoder.Build.cs b/Engine/Source/Runtime/AVEncoder/AVEncoder.Build.cs
--- a/Engine/Source/Runtime/AVEncoder/AVEncoder.Build.cs
+++ b/Engine/Source/Runtime/AVEncoder/AVEncoder.Build.cs
@@ -26,8 +26,8 @@ public class AVEncoder : ModuleRules
 
 		PrivateDependencyModuleNames.AddRange(new string[] {
 			"Engine",
-			"nvEncode"
-			// "Amf" // TODO waiting on cross platform Amf encoder
+			"nvEncode",
+			"Amf"
 		});
 
 		PublicDependencyModuleNames.AddRange(new string[] {
@@ -40,9 +40,19 @@ public class AVEncoder : ModuleRules
 		DynamicallyLoadedModuleNames.AddRange(new string[] {
 			// ... add any modules that your module loads dynamically here ...
 		});
+		
+		PrivateIncludePathModuleNames.AddRange(new string[] {
+			"VulkanRHI"
+		});
+
+		string EngineSourceDirectory = Path.GetFullPath(Target.RelativeEnginePath);
+		PrivateIncludePaths.Add(Path.Combine(EngineSourceDirectory, "Source/Runtime/VulkanRHI/Private"));
+		AddEngineThirdPartyPrivateStaticDependencies(Target, "Vulkan");
 
 		if (Target.Platform == UnrealTargetPlatform.Win64)
 		{
+			PublicDependencyModuleNames.Add("D3D12RHI");
+
 			// d3d to be able to use NVENC
 			PublicSystemLibraries.AddRange(new string[] {
 				"dxgi.lib",
@@ -51,10 +61,13 @@ public class AVEncoder : ModuleRules
 				"mfplat.lib",
 				"mfuuid.lib"
 			});
+
+			PrivateIncludePaths.Add(Path.Combine(EngineSourceDirectory, "Source/Runtime/VulkanRHI/Private/Windows"));
 		}
 		else if (Target.Platform == UnrealTargetPlatform.Linux)
 		{
-			PrivateDependencyModuleNames.Add("CUDA");			
+			PrivateDependencyModuleNames.Add("CUDA"); 
+			PrivateIncludePaths.Add(Path.Combine(EngineSourceDirectory, "Source/Runtime/VulkanRHI/Private/Linux"));
 		}
 		
 		// TEMPORARY: set this to zero for all platforms until CUDA TPS review clears
diff --git a/Engine/Source/Runtime/AVEncoder/Private/Decoders/Windows/VideoDecoderH264_Windows.cpp b/Engine/Source/Runtime/AVEncoder/Private/Decoders/Windows/VideoDecoderH264_Windows.cpp
--- a/Engine/Source/Runtime/AVEncoder/Private/Decoders/Windows/VideoDecoderH264_Windows.cpp
+++ b/Engine/Source/Runtime/AVEncoder/Private/Decoders/Windows/VideoDecoderH264_Windows.cpp
@@ -4,41 +4,13 @@
 
 #include <CoreMinimal.h>
 #include <HAL/Thread.h>
-#include "VideoCommon.h"
-/*********************************************************************************************************************/
-// Disable macro redefinition warning for compatibility with Windows SDK 8+
-// #pragma warning(push)
-// #pragma warning(disable : 4005)	// macro redefinition
-
-// #include "Windows/AllowWindowsPlatformTypes.h"
-// #include "Windows/PreWindowsApi.h"
-// 	#include <d3d11.h>
-// 	#include <mftransform.h>
-// 	#include <mfapi.h>
-// 	#include <mferror.h>
-// 	#include <mfidl.h>
-// 	#include <codecapi.h>
-// 	#include <shlwapi.h>
-// 	#include <mfreadwrite.h>
-// 	#include <d3d11_1.h>
-// 	#include <d3d12.h>
-// 	#include <dxgi1_4.h>
-// #include "Windows/PostWindowsApi.h"
-// #include "Windows/HideWindowsPlatformTypes.h"
-/*********************************************************************************************************************/
-//#include <Microsoft/COMPointer.h>
-#include "MicrosoftCommon.h"
 
-#include "Containers/Queue.h"
-#include "Misc/ScopeLock.h"
-#include "Templates/RefCounting.h"
+#include "MicrosoftCommon.h"
 
 #include "VideoDecoderCommon.h"
 #include "VideoDecoderAllocationTypes.h"
 #include "VideoDecoderUtilities.h"
 
-//#include "MicrosoftCommon.h"
-
 #define VERIFY_HR(FNcall,...)						\
 Result = FNcall;									\
 if (FAILED(Result))									\
diff --git a/Engine/Source/Runtime/AVEncoder/Private/Encoders/Amf/Amf_Common.cpp b/Engine/Source/Runtime/AVEncoder/Private/Encoders/Amf/Amf_Common.cpp
new file mode 100644
--- /dev/null
+++ b/Engine/Source/Runtime/AVEncoder/Private/Encoders/Amf/Amf_Common.cpp
@@ -0,0 +1,101 @@
+#include "Amf_Common.h"
+#include "AVEncoder.h"
+
+namespace AVEncoder
+{
+    FCriticalSection    FAmfCommon::ProtectSingleton;
+    FAmfCommon          FAmfCommon::Singleton;
+
+    // attempt to load Amf
+    FAmfCommon &FAmfCommon::Setup()
+    {
+        FScopeLock Guard(&ProtectSingleton);
+        if (!Singleton.bWasSetUp)
+        {
+            Singleton.bWasSetUp = true;
+            Singleton.SetupAmfFunctions();
+        }
+        return Singleton;
+    }
+
+    // shutdown - release loaded dll
+    void FAmfCommon::Shutdown()
+    {
+        FScopeLock Guard(&ProtectSingleton);
+        if (Singleton.bWasSetUp)
+        {
+            Singleton.bWasSetUp = false;
+            Singleton.bIsAvailable = false;
+            
+            if (Singleton.AmfContext)
+            {
+                Singleton.AmfContext->Terminate();
+                Singleton.AmfContext = nullptr;
+            }
+            
+            Singleton.AmfFactory = nullptr;
+            
+            if (Singleton.DllHandle)
+            {
+                FPlatformProcess::FreeDllHandle(Singleton.DllHandle);
+                Singleton.DllHandle = nullptr;
+            }
+        }
+    }
+
+    bool FAmfCommon::CreateEncoder(amf::AMFComponentPtr& outEncoder)
+    {
+		AMF_RESULT res = AmfFactory->CreateComponent(AmfContext, AMFVideoEncoderVCE_AVC, &outEncoder);
+		if (res != AMF_OK)
+		{
+			UE_LOG(LogAVEncoder, Error, TEXT("AMF failed to create Encoder component with code: %d"), res);
+			return false;
+		}
+        return true;
+    }
+
+    void FAmfCommon::SetupAmfFunctions()
+    {
+	    check(!bIsAvailable);
+
+		// Can't use Amf without and AMD GPU (also no point if its not the one RHI is using)
+		if (!IsRHIDeviceAMD())
+		{
+			return;
+		}
+        
+        DllHandle = FPlatformProcess::GetDllHandle(TEXT(AMF_DLL_NAMEA));
+
+        AMFInit_Fn AmfInitFn = (AMFInit_Fn)FPlatformProcess::GetDllExport(DllHandle, TEXT(AMF_INIT_FUNCTION_NAME));
+        if (AmfInitFn == nullptr)
+        {
+            return;
+        }
+
+        CHECK_AMF_RET(AmfInitFn(AMF_FULL_VERSION, &AmfFactory));
+
+        CHECK_AMF_RET(AmfFactory->CreateContext(&AmfContext));
+
+		// TODO this needs to get moved to lazy initialize when the encoder is actually called 
+		{
+			FString RHIName = GDynamicRHI->GetName();
+
+			if (RHIName == "D3D11")
+			{
+				AmfContext->InitDX11(GDynamicRHI->RHIGetNativeDevice());
+			}
+			else if (RHIName == "D3D12")
+			{
+				AMFContext2Ptr(AmfContext)->InitDX12(GDynamicRHI->RHIGetNativeDevice());
+			}
+			else if (RHIName == "Vulkan")
+			{
+				AMFContext1Ptr(AmfContext)->InitVulkan(GDynamicRHI->RHIGetNativeDevice());
+			}
+
+			bIsCtxInitialized = true;
+		}
+
+        bIsAvailable = true;
+    }
+}
\ No newline at end of file
diff --git a/Engine/Source/Runtime/AVEncoder/Private/Encoders/Amf/Amf_Common.h b/Engine/Source/Runtime/AVEncoder/Private/Encoders/Amf/Amf_Common.h
new file mode 100644
--- /dev/null
+++ b/Engine/Source/Runtime/AVEncoder/Private/Encoders/Amf/Amf_Common.h
@@ -0,0 +1,80 @@
+// Copyright Epic Games, Inc. All Rights Reserved.
+
+#pragma once
+
+#include <CoreMinimal.h>
+#include <HAL/Thread.h>
+
+#if PLATFORM_LINUX
+PRAGMA_DISABLE_OVERLOADED_VIRTUAL_WARNINGS
+#endif
+
+THIRD_PARTY_INCLUDES_START
+
+#if PLATFORM_WINDOWS
+#include "Windows/AllowWindowsPlatformTypes.h"
+#include "Windows/PreWindowsApi.h"
+#endif
+
+#include "core/Factory.h"
+#include "core/Interface.h"
+#include "components/VideoEncoderVCE.h"
+
+#if PLATFORM_WINDOWS
+#include "Windows/PostWindowsApi.h"
+#include "Windows/HideWindowsPlatformTypes.h"
+#endif
+
+THIRD_PARTY_INCLUDES_END
+
+#if PLATFORM_LINUX
+PRAGMA_ENABLE_OVERLOADED_VIRTUAL_WARNINGS
+#endif
+
+#define CHECK_AMF_RET(AMF_call)\
+{\
+	AMF_RESULT Res = AMF_call;\
+	if (!(Res== AMF_OK || Res==AMF_ALREADY_INITIALIZED))\
+	{\
+		UE_LOG(LogAVEncoder, Error, TEXT("`" #AMF_call "` failed with error code: %d"), Res);\
+		return;\
+	}\
+}
+
+namespace AVEncoder
+{
+	using namespace amf;
+
+    class FAmfCommon
+    {
+    public:
+        // attempt to load Amf
+        static FAmfCommon &Setup();
+
+        // shutdown - release loaded dll
+        static void Shutdown();
+
+        bool GetIsAvailable() const { return bIsAvailable; }
+
+		bool GetIsCtxInitialized() const { return bIsCtxInitialized; }
+		
+		bool CreateEncoder(amf::AMFComponentPtr& outEncoder);
+
+		AMFContextPtr GetContext() { return AmfContext; }
+
+		bool bIsCtxInitialized = false;
+    private:
+        FAmfCommon() = default;
+
+		void SetupAmfFunctions();
+
+        static FCriticalSection ProtectSingleton;
+        static FAmfCommon Singleton;
+
+        amf_handle DllHandle = nullptr;
+        AMFFactory *AmfFactory = nullptr;
+		AMFContextPtr AmfContext;
+        bool bIsAvailable = false;
+        bool bWasSetUp = false;
+    };
+}
\ No newline at end of file
diff --git a/Engine/Source/Runtime/AVEncoder/Private/Encoders/Amf/Amf_EncoderH264.cpp b/Engine/Source/Runtime/AVEncoder/Private/Encoders/Amf/Amf_EncoderH264.cpp
new file mode 100644
--- /dev/null
+++ b/Engine/Source/Runtime/AVEncoder/Private/Encoders/Amf/Amf_EncoderH264.cpp
@@ -0,0 +1,799 @@
+// Copyright Epic Games, Inc. All Rights Reserved.
+
+#include "Amf_EncoderH264.h"
+#include "VideoEncoderCommon.h"
+#include "CodecPacket.h"
+#include "AVEncoderDebug.h"
+
+#include <stdio.h>
+
+#define MAX_GPU_INDEXES 50
+#define DEFAULT_BITRATE 1000000u
+#define MAX_FRAMERATE_DIFF 0
+#define MIN_UPDATE_FRAMERATE_SECS 15
+
+#define AMF_VIDEO_ENCODER_START_TS L"StartTs"
+
+namespace
+{
+	AMF_VIDEO_ENCODER_RATE_CONTROL_METHOD_ENUM ConvertRateControlModeAMF(AVEncoder::FVideoEncoder::RateControlMode mode)
+	{
+		switch (mode)
+		{
+		case AVEncoder::FVideoEncoder::RateControlMode::CONSTQP: return AMF_VIDEO_ENCODER_RATE_CONTROL_METHOD_CONSTANT_QP;
+		case AVEncoder::FVideoEncoder::RateControlMode::VBR: return AMF_VIDEO_ENCODER_RATE_CONTROL_METHOD_LATENCY_CONSTRAINED_VBR;
+		default:
+		case AVEncoder::FVideoEncoder::RateControlMode::CBR: return AMF_VIDEO_ENCODER_RATE_CONTROL_METHOD_CBR;
+		}
+	}
+}
+
+namespace AVEncoder
+{
+	static bool GetEncoderInfo(FAmfCommon& Amf, FVideoEncoderInfo& EncoderInfo);
+
+	bool FVideoEncoderAmf_H264::GetIsAvailable(FVideoEncoderInputImpl& InInput, FVideoEncoderInfo& OutEncoderInfo)
+	{
+		FAmfCommon& Amf = FAmfCommon::Setup();
+		bool bIsAvailable = Amf.GetIsAvailable();
+		if (bIsAvailable)
+		{
+			OutEncoderInfo.CodecType = ECodecType::H264;
+		}
+		return bIsAvailable;
+	}
+
+	void FVideoEncoderAmf_H264::Register(FVideoEncoderFactory& InFactory)
+	{
+		FAmfCommon& Amf = FAmfCommon::Setup();
+		if (Amf.GetIsAvailable())
+		{
+			FVideoEncoderInfo	EncoderInfo;
+			if (GetEncoderInfo(Amf, EncoderInfo))
+			{
+				InFactory.Register(EncoderInfo, []() {
+					return TUniquePtr<FVideoEncoder>(new FVideoEncoderAmf_H264());
+				});
+			}
+		}
+	}
+
+	FVideoEncoderAmf_H264::FVideoEncoderAmf_H264()
+		: Amf(FAmfCommon::Setup())
+	{
+	}
+
+	FVideoEncoderAmf_H264::~FVideoEncoderAmf_H264()
+	{
+		Shutdown();
+	}
+
+	bool FVideoEncoderAmf_H264::Setup(TSharedRef<FVideoEncoderInput> input, FLayerConfig const& config)
+	{
+		if (!Amf.GetIsAvailable())
+		{
+			UE_LOG(LogVideoEncoder, Error, TEXT("Amf not avaliable"));
+			return false;
+		}
+
+		TSharedRef<FVideoEncoderInputImpl>	Input(StaticCastSharedRef<FVideoEncoderInputImpl>(input));
+
+		// TODO fix initializing contexts
+		FrameFormat = input->GetFrameFormat();
+		switch (FrameFormat)
+		{
+#if PLATFORM_WINDOWS
+		case AVEncoder::EVideoFrameFormat::D3D11_R8G8B8A8_UNORM:
+			EncoderDevice = Input->ForceD3D11InputFrames();
+			break;
+		case AVEncoder::EVideoFrameFormat::D3D12_R8G8B8A8_UNORM:
+			EncoderDevice = Input->ForceD3D11InputFrames(); // TODO use DX12 input when possible 
+			if (!Amf.bIsCtxInitialized)
+			{
+				Amf.GetContext()->InitDX11(EncoderDevice);
+				Amf.bIsCtxInitialized = true;
+			}
+			break;
+#endif
+		case AVEncoder::EVideoFrameFormat::VULKAN_R8G8B8A8_UNORM:
+			EncoderDevice = Input->GetVulkanDevice();
+			if (!Amf.bIsCtxInitialized)
+			{
+				AMFContext1Ptr(Amf.GetContext())->InitVulkan(EncoderDevice);
+				Amf.bIsCtxInitialized = true;
+			}
+			break;
+		case AVEncoder::EVideoFrameFormat::Undefined:
+		default:
+			UE_LOG(LogVideoEncoder, Error, TEXT("Frame format %s is not currently supported by Amf Encoder on this platform."), *ToString(FrameFormat));
+			return false;
+		}
+
+		if (!EncoderDevice)
+		{
+			UE_LOG(LogVideoEncoder, Error, TEXT("Amf needs an encoder device."));
+			return false;
+		}
+
+		auto mutableConfig = config;
+		if (mutableConfig.MaxFramerate == 0)
+			mutableConfig.MaxFramerate = 60;
+
+		return AddLayer(config);
+	}
+
+	FVideoEncoder::FLayer* FVideoEncoderAmf_H264::CreateLayer(uint32 layerIdx, FLayerConfig const& config)
+	{
+		auto const layer = new FAMFLayer(layerIdx, config, *this);
+		if (!layer->Setup())
+		{
+			delete layer;
+			return nullptr;
+		}
+		return layer;
+	}
+
+	void FVideoEncoderAmf_H264::DestroyLayer(FLayer* layer)
+	{
+		delete layer;
+	}
+
+	void FVideoEncoderAmf_H264::Encode(FVideoEncoderInputFrame const* frame, FEncodeOptions const& options)
+	{
+		// todo: reconfigure encoder
+		auto const amfFrame = static_cast<const FVideoEncoderInputFrameImpl*>(frame);
+		for (auto&& layer : Layers)
+		{
+			auto const amfLayer = static_cast<FAMFLayer*>(layer);
+			amfLayer->Encode(amfFrame, options);
+		}
+	}
+
+	void FVideoEncoderAmf_H264::Flush()
+	{
+		for (auto&& layer : Layers)
+		{
+			auto const amfLayer = static_cast<FAMFLayer*>(layer);
+			amfLayer->Flush();
+		}
+	}
+
+	void FVideoEncoderAmf_H264::Shutdown()
+	{
+		for (auto&& layer : Layers)
+		{
+			auto const amfLayer = static_cast<FAMFLayer*>(layer);
+			amfLayer->Shutdown();
+			DestroyLayer(amfLayer);
+		}
+		Layers.Reset();
+		StopEventThread();
+	}
+
+	void FVideoEncoderAmf_H264::OnEvent(void* InEvent, TUniqueFunction<void()>&& InCallback)
+	{
+#if PLATFORM_WINDOWS
+		FScopeLock Guard(&ProtectEventThread);
+		StartEventThread();
+		EventThreadWaitingFor.Emplace(FWaitForEvent(InEvent, MoveTemp(InCallback)));
+		SetEvent(EventThreadCheckEvent);
+#else
+		UE_LOG(LogVideoEncoder, Fatal, TEXT("FVideoEncoderAmf_H264::OnEvent should not be called as NVENC async mode only works on Windows!"))
+#endif
+	}
+
+	// TODO (M84FIX) de-windows this
+	void FVideoEncoderAmf_H264::StartEventThread()
+	{
+#if PLATFORM_WINDOWS
+		bExitEventThread = false;
+		if (!EventThread)
+		{
+			if (!EventThreadCheckEvent)
+			{
+				EventThreadCheckEvent = CreateEvent(nullptr, false, false, nullptr);
+				//EventThreadCheckEvent = FPlatformProcess::GetSynchEventFromPool(false);
+			}
+			EventThread = MakeUnique<FThread>(TEXT("NVENC_EncoderCommon"), [this]() { EventLoop(); });
+		}
+#else
+		UE_LOG(LogVideoEncoder, Fatal, TEXT("FVideoEncoderAmf_H264::StartEventThread should not be called as NVENC async mode only works on Windows!"))
+#endif
+	}
+
+	void FVideoEncoderAmf_H264::StopEventThread()
+	{
+#if PLATFORM_WINDOWS
+		FScopeLock				Guard(&ProtectEventThread);
+		TUniquePtr<FThread>		StopThread = MoveTemp(EventThread);
+		if (StopThread)
+		{
+			bExitEventThread = true;
+			SetEvent(EventThreadCheckEvent);
+			Guard.Unlock();
+			StopThread->Join();
+		}
+#else
+		UE_LOG(LogVideoEncoder, Fatal, TEXT("FVideoEncoderAmf_H264::StopEventThread should not be called as NVENC async mode only works on Windows!"))
+#endif
+	}
+
+#if PLATFORM_WINDOWS
+	void WindowsError(LPTSTR lpszFunction)
+	{
+		// Retrieve the system error message for the last-error code
+
+		LPVOID lpMsgBuf;
+		LPVOID lpDisplayBuf;
+		DWORD dw = GetLastError();
+
+		FormatMessage(
+			FORMAT_MESSAGE_ALLOCATE_BUFFER |
+			FORMAT_MESSAGE_FROM_SYSTEM |
+			FORMAT_MESSAGE_IGNORE_INSERTS,
+			NULL,
+			dw,
+			MAKELANGID(LANG_NEUTRAL, SUBLANG_DEFAULT),
+			(LPTSTR)&lpMsgBuf,
+			0, NULL);
+
+		lpDisplayBuf = (LPVOID)LocalAlloc(LMEM_ZEROINIT, (lstrlen((LPCTSTR)lpMsgBuf) + lstrlen((LPCTSTR)lpszFunction) + 40) * sizeof(TCHAR));
+		UE_LOG(LogVideoEncoder, Error, TEXT("%s failed with error %d: %s"), lpszFunction, dw, lpMsgBuf);
+
+		LocalFree(lpMsgBuf);
+		LocalFree(lpDisplayBuf);
+	}
+#endif
+
+	// TODO (M84FIX) de-windows this
+	void FVideoEncoderAmf_H264::EventLoop()
+	{
+#if PLATFORM_WINDOWS
+		while (true)
+		{
+			// Make an empty array of events to wait for
+			void* EventsToWaitFor[MAXIMUM_WAIT_OBJECTS];
+			int				NumEventsToWaitFor = 0;
+
+			{
+				// Guard the thread and see if we should terminate
+				FScopeLock		Guard(&ProtectEventThread);
+				if (bExitEventThread)
+				{
+					break;
+				}
+
+				// collect events to wait for
+				EventsToWaitFor[NumEventsToWaitFor++] = EventThreadCheckEvent;
+				for (const FWaitForEvent& WaitFor : EventThreadWaitingFor)
+				{
+					check(NumEventsToWaitFor < MAXIMUM_WAIT_OBJECTS);
+					EventsToWaitFor[NumEventsToWaitFor++] = WaitFor.Key;
+				}
+			}
+			// wait for the events
+			DWORD	WaitResult = WaitForMultipleObjects(NumEventsToWaitFor, EventsToWaitFor, false, INFINITE);
+			if (WaitResult >= WAIT_OBJECT_0 && WaitResult < (WAIT_OBJECT_0 + NumEventsToWaitFor))
+			{
+				// Guard the thread and get the event that was triggered
+				FScopeLock		Guard(&ProtectEventThread);
+				void* EventTriggered = EventsToWaitFor[WaitResult - WAIT_OBJECT_0];
+
+				// loop through the the events then callback on the one that was triggered (could this be better as a)
+				for (int32 Index = 0; Index < EventThreadWaitingFor.Num(); ++Index)
+				{
+					if (EventThreadWaitingFor[Index].Key == EventTriggered)
+					{
+						TUniqueFunction<void()>	Callback = MoveTemp(EventThreadWaitingFor[Index].Value);
+						EventThreadWaitingFor.RemoveAtSwap(Index);
+						Guard.Unlock();
+						Callback();
+						break;
+					}
+				}
+			}
+			else if (WaitResult >= WAIT_ABANDONED_0 && WaitResult < (WAIT_ABANDONED_0 + NumEventsToWaitFor))
+			{
+			}
+			else if (WaitResult == WAIT_TIMEOUT)
+			{
+			}
+			else if (WaitResult == WAIT_FAILED)
+			{
+				WindowsError(TEXT("WaitForMultipleObjects"));
+			}
+		}
+#else
+		UE_LOG(LogVideoEncoder, Fatal, TEXT("Amf_EncoderH264::EventLoop should not be called as NVENC async mode only works on Windows!"))
+#endif
+	}
+
+	// --- Amf_EncoderH264::FLayer ------------------------------------------------------------
+	FVideoEncoderAmf_H264::FAMFLayer::FAMFLayer(uint32 layerIdx, FLayerConfig const& config, FVideoEncoderAmf_H264& encoder)
+		: FLayer(config)
+		, Encoder(encoder)
+		, Amf(FAmfCommon::Setup())
+		, CodecGUID()
+		, LayerIndex(layerIdx)
+	{
+	}
+
+	FVideoEncoderAmf_H264::FAMFLayer::~FAMFLayer()
+	{
+	}
+
+	bool FVideoEncoderAmf_H264::FAMFLayer::Setup()
+	{
+		return CreateSession() && CreateInitialConfig();
+	}
+
+	bool FVideoEncoderAmf_H264::FAMFLayer::CreateSession()
+	{
+		if (AmfEncoder == NULL)
+		{
+			Amf.CreateEncoder(AmfEncoder);
+
+			// TODO maybe setup memory usage here though it looks like Amf doesnt need to have this pre-configured
+		}
+
+		return AmfEncoder != NULL;
+	}
+
+	// TODO (M84FIX) need to parameterize this
+	bool FVideoEncoderAmf_H264::FAMFLayer::CreateInitialConfig()
+	{
+		AMF_RESULT Result = AMF_OK;
+
+		// TODO this might be a double up
+		Result = AmfEncoder->SetProperty(AMF_VIDEO_ENCODER_USAGE, AMF_VIDEO_ENCODER_USAGE_ULTRA_LOW_LATENCY);
+		Result = AmfEncoder->SetProperty(AMF_VIDEO_ENCODER_LOWLATENCY_MODE, true);
+
+		// TODO set to baseline to match NVENC this seems to allow the most devices to decode
+		Result = AmfEncoder->SetProperty(AMF_VIDEO_ENCODER_PROFILE, AMF_VIDEO_ENCODER_PROFILE_BASELINE);
+		// Result = AmfEncoder->SetProperty(AMF_VIDEO_ENCODER_PROFILE_LEVEL, 51);
+
+		Result = AmfEncoder->SetProperty(AMF_VIDEO_ENCODER_FRAMERATE, CurrentConfig.MaxFramerate);
+
+		Result = AmfEncoder->SetProperty(AMF_VIDEO_ENCODER_RATE_CONTROL_METHOD, ConvertRateControlModeAMF(CurrentConfig.RateControlMode));
+
+		// TODO will this do the same things as RcParams.enableMinQP/enableMaxQP = false in NVENC?
+		Result = AmfEncoder->SetProperty(AMF_VIDEO_ENCODER_MIN_QP, CurrentConfig.QPMin == -1 ? 0 : CurrentConfig.QPMin);
+		Result = AmfEncoder->SetProperty(AMF_VIDEO_ENCODER_MAX_QP, CurrentConfig.QPMax == -1 ? 51 : CurrentConfig.QPMax);
+
+		Result = AmfEncoder->SetProperty(AMF_VIDEO_ENCODER_PEAK_BITRATE, CurrentConfig.MaxBitrate > -1 ? CurrentConfig.MaxBitrate : DEFAULT_BITRATE);
+		Result = AmfEncoder->SetProperty(AMF_VIDEO_ENCODER_TARGET_BITRATE, CurrentConfig.TargetBitrate > -1 ? CurrentConfig.TargetBitrate : DEFAULT_BITRATE);
+
+		// TODO Amf enables B frames by default NVENC does not should test both options
+		Result = AmfEncoder->SetProperty(AMF_VIDEO_ENCODER_B_PIC_PATTERN, 0);
+
+		Result = AmfEncoder->Init(AMF_SURFACE_YUV420P, CurrentConfig.Width, CurrentConfig.Height);
+
+		return Result == AMF_OK;
+	}
+
+	void FVideoEncoderAmf_H264::FAMFLayer::MaybeReconfigure(TSharedPtr<FInputOutput> buffer)
+	{
+		FScopeLock lock(&ConfigMutex);
+		if (NeedsReconfigure)
+		{
+			// TODO this is likely to work differently in Amf and should be changed accordingly
+			// currently it mirrors NVENC to have paridy
+
+			// Changing the framerate forces the generation of a new keyframe so we only do it when either:
+			// 1) we need a keyframe anyway
+			// 2) there is a big difference in framerate
+			// 3) the framerate has changed and we haven't sent a Keyframe recently
+			AMFRate CurrentEncoderFramerate;
+			AmfEncoder->GetProperty(AMF_VIDEO_ENCODER_FRAMERATE, &CurrentEncoderFramerate);
+			uint32 const FrameRateDiff = CurrentConfig.MaxFramerate > CurrentEncoderFramerate.num ? CurrentConfig.MaxFramerate - CurrentEncoderFramerate.num : CurrentEncoderFramerate.num - CurrentConfig.MaxFramerate;
+
+			if (bForceNextKeyframe ||
+				FrameRateDiff > MAX_FRAMERATE_DIFF ||
+				(CurrentConfig.MaxFramerate != CurrentEncoderFramerate.num && (FDateTime::UtcNow() - LastKeyFrameTime).GetSeconds() > MIN_UPDATE_FRAMERATE_SECS))
+			{
+				if (AmfEncoder->SetProperty(AMF_VIDEO_ENCODER_FRAMERATE, CurrentConfig.MaxFramerate) != AMF_OK)
+				{
+					UE_LOG(LogVideoEncoder, Error, TEXT("Amf failed to set encoder Framerate"));
+				}
+
+				if (buffer->Surface->SetProperty(AMF_VIDEO_ENCODER_FORCE_PICTURE_TYPE, AMF_VIDEO_ENCODER_PICTURE_TYPE_IDR) != AMF_OK)
+				{
+					UE_LOG(LogVideoEncoder, Error, TEXT("Amf failed to force IDR picture type"));
+				}
+			}
+			bForceNextKeyframe = false;
+
+			if (AmfEncoder->SetProperty(AMF_VIDEO_ENCODER_FILLER_DATA_ENABLE, CurrentConfig.FillData) != AMF_OK)
+			{
+				UE_LOG(LogVideoEncoder, Error, TEXT("Amf failed to enable filler data to maintain CBR"));
+			}
+
+			if (AmfEncoder->SetProperty(AMF_VIDEO_ENCODER_RATE_CONTROL_METHOD, ConvertRateControlModeAMF(CurrentConfig.RateControlMode)) != AMF_OK)
+			{
+				UE_LOG(LogVideoEncoder, Error, TEXT("Amf failed to set rate control method"));
+			}
+
+			if (AmfEncoder->SetProperty(AMF_VIDEO_ENCODER_PEAK_BITRATE, CurrentConfig.MaxBitrate) != AMF_OK)
+			{
+				UE_LOG(LogVideoEncoder, Error, TEXT("Amf failed to set max bitrate"));
+			}
+
+			// multi pass?
+
+			if (AmfEncoder->SetProperty(AMF_VIDEO_ENCODER_MIN_QP, CurrentConfig.QPMin) != AMF_OK)
+			{
+				UE_LOG(LogVideoEncoder, Error, TEXT("Amf failed to set min qp"));
+			}
+			
+			if (AmfEncoder->SetProperty(AMF_VIDEO_ENCODER_TARGET_BITRATE, CurrentConfig.TargetBitrate) != AMF_OK)
+			{
+				UE_LOG(LogVideoEncoder, Error, TEXT("Amf failed to set target bitrate"));
+			}
+
+			// Bitrate generated by NVENC is not very stable when the scene doesnt have a lot of movement
+			// outputPictureTimingSEI enables the filling data when using CBR so that the bitrate generated
+			// is much closer to the requested one and bandwidth estimation algorithms can work better.
+			// Otherwise in a static scene it can send 50kbps when configuring 300kbps and it will never ramp up.
+			if (CurrentConfig.TargetBitrate < 5000000)
+			{
+				if (AmfEncoder->SetProperty(AMF_VIDEO_ENCODER_MIN_QP, 0) != AMF_OK)
+				{
+					UE_LOG(LogVideoEncoder, Error, TEXT("Amf failed to disable min qp"));
+				}
+			}
+			else
+			{
+				if (AmfEncoder->SetProperty(AMF_VIDEO_ENCODER_FILLER_DATA_ENABLE, true) != AMF_OK)
+				{
+					UE_LOG(LogVideoEncoder, Error, TEXT("Amf failed to enable filler data to maintain CBR"));
+				}
+
+				if (AmfEncoder->SetProperty(AMF_VIDEO_ENCODER_MIN_QP, 20) != AMF_OK)
+				{
+					UE_LOG(LogVideoEncoder, Error, TEXT("Amf failed to set min qp to 20"));
+				}
+			}
+
+			NeedsReconfigure = false;
+		}
+	}
+
+	void FVideoEncoderAmf_H264::FAMFLayer::Encode(FVideoEncoderInputFrameImpl const* frame, FEncodeOptions const& options)
+	{
+		TSharedPtr<FInputOutput> Buffer = GetOrCreateSurface(frame);
+
+		if (Buffer)
+		{
+			Buffer->Surface->SetPts(FTimespan::FromSeconds(FPlatformTime::Seconds()).GetTicks());
+
+			bForceNextKeyframe = options.bForceKeyFrame;
+			MaybeReconfigure(Buffer);
+
+			if (Buffer.IsValid())
+			{
+				AMF_RESULT Result = AmfEncoder->SubmitInput(Buffer->Surface);
+
+				if (Result == AMF_NEED_MORE_INPUT)
+				{
+				}
+				else if (Result != AMF_OK)
+				{
+					UE_LOG(LogVideoEncoder, Error, TEXT("Amf submit error with %d"), Result);
+					// release input frame
+					if (Buffer->SourceFrame)
+					{
+						Buffer->SourceFrame->Release();
+					}
+				}
+			}
+		}
+	}
+
+	void FVideoEncoderAmf_H264::FAMFLayer::Flush()
+	{
+		AmfEncoder->Flush();
+	}
+
+	void FVideoEncoderAmf_H264::FAMFLayer::Shutdown()
+	{
+		Flush();
+		if (AmfEncoder != NULL)
+		{
+			AmfEncoder->Terminate();
+			AmfEncoder = NULL;
+		}
+	}
+
+	void FVideoEncoderAmf_H264::FAMFLayer::ProcessNextPendingFrame()
+	{
+		if (Encoder.OnEncodedPacket)
+		{
+			while (true)
+			{
+				amf::AMFDataPtr data;
+				AMF_RESULT Result = AmfEncoder->QueryOutput(&data);
+
+				if (Result == AMF_OK && data != NULL)
+				{
+					AMFBufferPtr OutBuffer(data);
+
+					// create packet with buffer contents
+					FCodecPacketImpl	Packet;
+
+					Packet.PTS = data->GetPts();
+					Packet.Data = static_cast<const uint8*>(OutBuffer->GetNative());
+					Packet.DataSize = OutBuffer->GetSize();
+					uint32 PictureType = AMF_VIDEO_ENCODER_PICTURE_TYPE_NONE;
+					if (OutBuffer->GetProperty(AMF_VIDEO_ENCODER_FORCE_PICTURE_TYPE, &PictureType) != AMF_OK)
+					{
+						UE_LOG(LogVideoEncoder, Fatal, TEXT("Amf failed to get picture type."));
+					}
+					else if (PictureType == AMF_VIDEO_ENCODER_PICTURE_TYPE_IDR)
+					{
+						UE_LOG(LogVideoEncoder, Verbose, TEXT("Generated IDR Frame"));
+						Packet.IsKeyFrame = true;
+					}
+
+					if (OutBuffer->GetProperty(AMF_VIDEO_ENCODER_STATISTIC_AVERAGE_QP, &Packet.VideoQP) != AMF_OK)
+					{
+						UE_LOG(LogVideoEncoder, Fatal, TEXT("Amf failed to get average QP."));
+					}
+
+					if (OutBuffer->GetProperty(AMF_VIDEO_ENCODER_START_TS, &Packet.Timings.StartTs) != AMF_OK)
+					{
+						UE_LOG(LogVideoEncoder, Fatal, TEXT("Amf failed to get average QP."));
+					}
+
+					Packet.Timings.FinishTs = FTimespan::FromSeconds(FPlatformTime::Seconds());
+					Packet.Framerate = CurrentConfig.MaxFramerate;
+
+					Encoder.OnEncodedPacket(LayerIndex, nullptr, Packet); // TODO probs should not be null might need a refactor to instead pass the size as thats what its needed for
+				}
+				else
+				{
+					break;
+				}
+			}
+		}
+	}
+
+	template<class T>
+	bool FVideoEncoderAmf_H264::FAMFLayer::GetCapability(const TCHAR* CapToQuery, T& OutCap) const
+	{
+		amf::AMFCapsPtr EncoderCaps;
+		
+		if (AmfEncoder->GetCaps(&EncoderCaps) != AMF_OK)
+		{
+			return false;
+		}
+
+		if (EncoderCaps->GetProperty(CapToQuery, &OutCap))
+		{
+			return false;
+		}
+
+		return true;
+	}
+
+	TSharedPtr<FVideoEncoderAmf_H264::FAMFLayer::FInputOutput> FVideoEncoderAmf_H264::FAMFLayer::GetOrCreateSurface(const FVideoEncoderInputFrameImpl* InFrame)
+	{
+		void* TextureToCompress = nullptr;
+
+		switch (InFrame->GetFormat())
+		{
+#if PLATFORM_WINDOWS
+		case EVideoFrameFormat::D3D11_R8G8B8A8_UNORM:
+			TextureToCompress = InFrame->GetD3D11().EncoderTexture;
+			break;
+		case EVideoFrameFormat::D3D12_R8G8B8A8_UNORM:
+			TextureToCompress = InFrame->GetD3D12().EncoderTexture;
+			break;
+#endif
+		case EVideoFrameFormat::VULKAN_R8G8B8A8_UNORM:
+			TextureToCompress = InFrame->GetVulkan().EncoderTexture;
+			break;
+		case EVideoFrameFormat::Undefined:
+		default:
+			UE_LOG(LogVideoEncoder, Error, TEXT("Video Frame Format %s not supported by Amf on this platform."), *ToString(InFrame->GetFormat()));
+			break;
+		}
+
+		if (!TextureToCompress)
+		{
+			UE_LOG(LogVideoEncoder, Fatal, TEXT("Got passed a null pointer."));
+				return nullptr;
+		}
+
+		// Check if texture already has a buffer surface
+		TSharedPtr<FInputOutput> Buffer = nullptr;
+		for (TSharedPtr<FInputOutput> SearchBuffer : CreatedSurfaces)
+		{
+			if (SearchBuffer->TextureToCompress == TextureToCompress)
+			{
+				Buffer = SearchBuffer;
+				break;
+			}
+		}
+
+		// if texture does not already have a buffer surface create one
+		if (!Buffer)
+		{
+			if (!CreateSurface(Buffer, InFrame, TextureToCompress))
+			{
+				InFrame->Release();
+				UE_LOG(LogVideoEncoder, Error, TEXT("Amf failed to create buffer."));
+			}
+			else
+			{
+				CreatedSurfaces.Push(Buffer);
+			}
+		}
+
+		return Buffer;
+	}
+
+	bool FVideoEncoderAmf_H264::FAMFLayer::CreateSurface(TSharedPtr<FVideoEncoderAmf_H264::FAMFLayer::FInputOutput>& OutBuffer, const FVideoEncoderInputFrameImpl* SourceFrame, void* TextureToCompress)
+	{
+		AMF_RESULT Result = AMF_OK;
+
+		OutBuffer = MakeShared<FInputOutput>();
+		OutBuffer->SourceFrame = SourceFrame;
+		OutBuffer->TextureToCompress = TextureToCompress;
+
+		if (TextureToCompress)
+		{
+			switch (SourceFrame->GetFormat())
+			{
+#if PLATFORM_WINDOWS
+			case EVideoFrameFormat::D3D11_R8G8B8A8_UNORM:
+				Result = Amf.GetContext()->CreateSurfaceFromDX11Native(TextureToCompress, &(OutBuffer->Surface), OutBuffer.Get());
+				break;
+			case EVideoFrameFormat::D3D12_R8G8B8A8_UNORM:
+				Result = amf::AMFContext2Ptr(Amf.GetContext())->CreateSurfaceFromDX12Native(TextureToCompress, &(OutBuffer->Surface), OutBuffer.Get());
+				break;
+#endif
+			case EVideoFrameFormat::VULKAN_R8G8B8A8_UNORM:
+				Result = amf::AMFContext2Ptr(Amf.GetContext())->CreateSurfaceFromVulkanNative(TextureToCompress, &(OutBuffer->Surface), OutBuffer.Get());
+				break;
+			case EVideoFrameFormat::Undefined:
+			default:
+				UE_LOG(LogVideoEncoder, Error, TEXT("Video format %s not inplemented for Amf on this platform"), *ToString(SourceFrame->GetFormat()));
+				break;
+			}
+		}
+		else
+		{
+			UE_LOG(LogVideoEncoder, Error, TEXT("Amf recieved nullptr to texture."));
+			return false;
+		}
+
+		return Result == AMF_OK;
+	}
+
+
+	static void* CreateEncoderSession(FAmfCommon& Amf)
+	{
+		void* EncoderSession = nullptr;
+		return EncoderSession;
+	}
+
+	static bool GetEncoderSupportedProfiles(AMFCapsPtr EncoderCaps, uint32& OutSupportedProfiles)
+	{
+		uint32 maxProfile;
+		if (EncoderCaps->GetProperty(AMF_VIDEO_ENCODER_CAP_MAX_PROFILE, &maxProfile) != AMF_OK)
+		{
+			return false;
+		}
+
+		if (maxProfile >= AMF_VIDEO_ENCODER_PROFILE_BASELINE)
+		{
+			OutSupportedProfiles |= H264Profile_Baseline;
+		}
+
+		if (maxProfile >= AMF_VIDEO_ENCODER_PROFILE_MAIN)
+		{
+			OutSupportedProfiles |= H264Profile_Main;
+		}
+
+		if (maxProfile >= AMF_VIDEO_ENCODER_PROFILE_HIGH)
+		{
+			OutSupportedProfiles |= H264Profile_High;
+		}
+
+		if (maxProfile >= AMF_VIDEO_ENCODER_PROFILE_CONSTRAINED_BASELINE)
+		{
+			OutSupportedProfiles |= H264Profile_ConstrainedBaseline;
+		}
+
+		if (maxProfile >= AMF_VIDEO_ENCODER_PROFILE_CONSTRAINED_HIGH)
+		{
+			OutSupportedProfiles |= H264Profile_ConstrainedHigh;
+		}
+
+		return true;
+	}
+
+	static bool GetEncoderSupportedInputFormats(AMFIOCapsPtr IOCaps, TArray<EVideoFrameFormat>& OutSupportedInputFormats)
+	{
+		// TODO check if we actually need to query Amf for this
+
+#if PLATFORM_WINDOWS
+		OutSupportedInputFormats.Push(EVideoFrameFormat::D3D11_R8G8B8A8_UNORM);
+		OutSupportedInputFormats.Push(EVideoFrameFormat::D3D12_R8G8B8A8_UNORM);
+#endif
+		OutSupportedInputFormats.Push(EVideoFrameFormat::VULKAN_R8G8B8A8_UNORM);
+
+		return true;
+	}
+
+
+	static bool GetEncoderInfo(FAmfCommon& AMF, FVideoEncoderInfo& EncoderInfo)
+	{
+		bool bSuccess = true;
+		// create a temporary encoder component
+		AMFComponentPtr AmfEncoder = nullptr;
+
+		{
+			FString RHIName = FString(GDynamicRHI->GetName());
+			void* EncoderDevice = GDynamicRHI->RHIGetNativeDevice();
+
+			// TODO replace GDynamicRHI with a passed in handle to the RHI Device so that this can be used agnostically
+#if PLATFORM_WINDOWS
+			if (RHIName == FString("D3D12"))
+			{
+				AMFContext2Ptr(AMF.GetContext())->InitDX12(EncoderDevice);
+			}
+
+			if (RHIName == FString("D3D11"))
+			{
+				AMF.GetContext()->InitDX11(EncoderDevice);
+			}
+#endif			
+			if (RHIName == FString("Vulkan"))
+			{
+				AMFContext1Ptr(AMF.GetContext())->InitVulkan(EncoderDevice);
+			}
+
+			if (EncoderDevice == nullptr || !AMF.CreateEncoder(AmfEncoder))
+			{
+				UE_LOG(LogVideoEncoder, Error, TEXT("RHI not intitialized or not supported by Amf on this platform."));
+				return false;
+			}
+		}
+
+		EncoderInfo.CodecType = ECodecType::H264;
+		EncoderInfo.MaxWidth;
+		EncoderInfo.MaxHeight;
+
+		AMFCapsPtr EncoderCaps;
+		AmfEncoder->GetCaps(&EncoderCaps);
+
+		AMFIOCapsPtr InputCaps;
+		EncoderCaps->GetInputCaps(&InputCaps);
+
+		uint32 LevelMin = 1;
+		uint32 LevelMax;
+		if (EncoderCaps->GetProperty(AMF_VIDEO_ENCODER_CAP_MAX_LEVEL, &LevelMax) == AMF_OK)
+		{
+			EncoderInfo.H264.MinLevel = (LevelMin > 9) ? LevelMin : 9;							// Like NVENC we hard min at 9
+			EncoderInfo.H264.MaxLevel = (LevelMax < 9) ? 9 : (LevelMax > 52) ? 52 : LevelMax;	// Like NVENC we hard min at 52
+		}
+		else
+		{
+			UE_LOG(LogVideoEncoder, Error, TEXT("Failed to query min/max h264 level supported by Amf (reported min/max=%d/%d)."), LevelMin, LevelMax);
+			bSuccess = false;
+		}
+
+		if (!GetEncoderSupportedProfiles(EncoderCaps, EncoderInfo.H264.SupportedProfiles) ||
+			!GetEncoderSupportedInputFormats(InputCaps, EncoderInfo.SupportedInputFormats))
+		{
+			bSuccess = false;
+		}
+
+		// destroy encoder session
+		AmfEncoder->Terminate();
+
+		return bSuccess;
+	}
+
+} /* namespace AVEncoder */
\ No newline at end of file
diff --git a/Engine/Source/Runtime/AVEncoder/Private/Encoders/Amf/Amf_EncoderH264.h b/Engine/Source/Runtime/AVEncoder/Private/Encoders/Amf/Amf_EncoderH264.h
new file mode 100644
--- /dev/null
+++ b/Engine/Source/Runtime/AVEncoder/Private/Encoders/Amf/Amf_EncoderH264.h
@@ -0,0 +1,125 @@
+// Copyright Epic Games, Inc. All Rights Reserved.
+
+#pragma once
+
+#include "Amf_Common.h"
+
+#include "VideoEncoderFactory.h"
+#include "VideoEncoderInputImpl.h"
+
+
+namespace AVEncoder
+{
+	using namespace amf;
+
+    class FVideoEncoderAmf_H264 : public FVideoEncoder
+    {
+    public:
+        virtual ~FVideoEncoderAmf_H264() override;
+
+        // query whether or not encoder is supported and available
+        static bool GetIsAvailable(FVideoEncoderInputImpl &InInput, FVideoEncoderInfo &OutEncoderInfo);
+
+        // register encoder with video encoder factory
+        static void Register(FVideoEncoderFactory &InFactory);
+
+        bool Setup(TSharedRef<FVideoEncoderInput> input, FLayerConfig const& config) override;
+        void Encode(FVideoEncoderInputFrame const* frame, FEncodeOptions const& options) override;
+        void Flush();
+        void Shutdown() override;
+
+    protected:
+        FLayer* CreateLayer(uint32 layerIdx, FLayerConfig const& config) override;
+        void DestroyLayer(FLayer *layer) override;
+
+    private:
+        FVideoEncoderAmf_H264();
+
+		class FAMFLayer : public FLayer
+		{
+		public:
+			FAMFLayer(uint32 layerIdx, FLayerConfig const& config, FVideoEncoderAmf_H264& encoder);
+			~FAMFLayer();
+
+			bool Setup();
+			bool CreateSession();
+			bool CreateInitialConfig();
+
+			template<class T> 
+			bool GetProperty(const TCHAR* PropertyToQuery, T& outProperty, const T& (*func)(const AMFVariantStruct*)) const;
+
+			template<class T>
+			bool GetCapability(const TCHAR* CapToQuery, T& OutCap) const;
+
+			void Encode(FVideoEncoderInputFrameImpl const* frame, FEncodeOptions const& options);
+			void Flush();
+			void Shutdown();
+			void UpdateBitrate(uint32 InMaxBitRate, uint32 InTargetBitRate);
+			void UpdateResolution(uint32 InMaxBitRate, uint32 InTargetBitRate);
+
+			FVideoEncoderAmf_H264& Encoder;
+			FAmfCommon& Amf;
+			GUID CodecGUID;
+			uint32 LayerIndex;
+			AMFComponentPtr AmfEncoder;
+			bool bAsyncMode = true;
+			FDateTime LastKeyFrameTime = 0;
+			bool bForceNextKeyframe = false;
+
+			class FInputOutput : public AMFSurfaceObserver
+			{
+			public:
+				virtual void AMF_STD_CALL OnSurfaceDataRelease(AMFSurface* pSurface) override
+				{
+					SourceFrame->Release();
+				}
+
+				virtual ~FInputOutput()
+				{
+					SourceFrame->Release();
+					Surface->Release();
+				}
+
+				const FVideoEncoderInputFrameImpl* SourceFrame;
+				void* TextureToCompress;
+				AMFSurfacePtr Surface;
+			};
+
+			void MaybeReconfigure(TSharedPtr<FInputOutput> buffer);
+
+			TSharedPtr<FInputOutput> GetOrCreateSurface(const FVideoEncoderInputFrameImpl* InFrame);
+			bool CreateSurface(TSharedPtr<FInputOutput>& OutBuffer, const FVideoEncoderInputFrameImpl* SourceFrame, void* TextureToCompress);
+
+			void ProcessNextPendingFrame();
+
+			TArray<TSharedPtr<FInputOutput>> CreatedSurfaces;
+			FCriticalSection ProtectedWaitingForPending;
+			bool WaitingForPendingActive = false;
+			FThreadSafeBool bUpdateConfig = false;
+		};
+
+		FAmfCommon& Amf;
+		EVideoFrameFormat FrameFormat = EVideoFrameFormat::Undefined;
+		void* EncoderDevice;
+
+		uint32 MaxFramerate = 0;
+		int32 MinQP = -1;
+		RateControlMode RateMode = RateControlMode::CBR;
+		bool FillData = false;
+
+		// event thread for amf async (supported on all platforms)
+		void OnEvent(void* InEvent, TUniqueFunction<void()>&& InCallback);
+		void StartEventThread();
+		void StopEventThread();
+		void EventLoop();
+
+		TUniquePtr<FThread> EventThread;
+		FCriticalSection ProtectEventThread;
+		bool bExitEventThread = false;
+
+		void* EventThreadCheckEvent = nullptr;
+
+		using FWaitForEvent = TPair<void*, TUniqueFunction<void()>>;
+		TArray<FWaitForEvent> EventThreadWaitingFor;
+    };
+}
\ No newline at end of file
diff --git a/Engine/Source/Runtime/AVEncoder/Private/Encoders/NVENC/NVENC_Common.cpp b/Engine/Source/Runtime/AVEncoder/Private/Encoders/NVENC/NVENC_Common.cpp
--- a/Engine/Source/Runtime/AVEncoder/Private/Encoders/NVENC/NVENC_Common.cpp
+++ b/Engine/Source/Runtime/AVEncoder/Private/Encoders/NVENC/NVENC_Common.cpp
@@ -43,7 +43,13 @@ void FNVENCCommon::Shutdown()
 
 void FNVENCCommon::SetupNVENCFunctions()
 {
-	bIsAvailable = false;
+	check(!bIsAvailable);
+
+	// Can't use NVENC without a NVIDIA GPU (also no point if its not the one RHI is using)
+	if (!IsRHIDeviceNVIDIA())
+	{
+		return;
+	}
 
 	// clear function call table
 	FMemory::Memzero(static_cast<NV_ENCODE_API_FUNCTION_LIST*>(this), sizeof(NV_ENCODE_API_FUNCTION_LIST));
