Merge in PixelStreaming, AVEncoder, and SignallingServer changes from TensorWorks local repository

diff --git a/Engine/Source/Runtime/AVEncoder/Private/Encoders/NVENC/NVENC_EncoderH264.cpp b/Engine/Source/Runtime/AVEncoder/Private/Encoders/NVENC/NVENC_EncoderH264.cpp
--- a/Engine/Source/Runtime/AVEncoder/Private/Encoders/NVENC/NVENC_EncoderH264.cpp
+++ b/Engine/Source/Runtime/AVEncoder/Private/Encoders/NVENC/NVENC_EncoderH264.cpp
@@ -17,197 +17,11 @@
 #define MAX_GPU_INDEXES 50
 #define DEFAULT_BITRATE 1000000u
 #define MAX_FRAMERATE_DIFF 0
-#define MIN_UPDATE_FRAMERATE_SECS 15
+#define MIN_UPDATE_FRAMERATE_SECS 5
 
-namespace AVEncoder
+namespace
 {
-
-	static bool GetEncoderInfo(FNVENCCommon& NVENC, FVideoEncoderInfo& EncoderInfo);
-
-
-	bool FVideoEncoderNVENC_H264::GetIsAvailable(FVideoEncoderInputImpl& InInput, FVideoEncoderInfo& OutEncoderInfo)
-	{
-		FNVENCCommon& NVENC = FNVENCCommon::Setup();
-		bool					bIsAvailable = NVENC.GetIsAvailable();
-		if (bIsAvailable)
-		{
-			OutEncoderInfo.CodecType = ECodecType::H264;
-
-		}
-		return bIsAvailable;
-	}
-
-	void FVideoEncoderNVENC_H264::Register(FVideoEncoderFactory& InFactory)
-	{
-		FNVENCCommon& NVENC = FNVENCCommon::Setup();
-		if (NVENC.GetIsAvailable())
-		{
-			FVideoEncoderInfo	EncoderInfo;
-			if (GetEncoderInfo(NVENC, EncoderInfo))
-			{
-				InFactory.Register(EncoderInfo, []() {
-					return TUniquePtr<FVideoEncoder>(new FVideoEncoderNVENC_H264());
-					});
-			}
-		}
-	}
-
-	TUniquePtr<FVideoEncoderNVENC_H264> FVideoEncoderNVENC_H264::Create(FVideoEncoderInputImpl& InInput)
-	{
-		return TUniquePtr<FVideoEncoderNVENC_H264>(new FVideoEncoderNVENC_H264());
-	}
-
-	FVideoEncoderNVENC_H264::FVideoEncoderNVENC_H264()
-		: NVENC(FNVENCCommon::Setup())
-	{
-	}
-
-	FVideoEncoderNVENC_H264::~FVideoEncoderNVENC_H264()
-	{
-		Shutdown();
-	}
-
-	bool FVideoEncoderNVENC_H264::Setup(TSharedRef<FVideoEncoderInput> InInput, const FInit& InInit)
-	{
-		if (!NVENC.GetIsAvailable())
-		{
-			UE_LOG(LogVideoEncoder, Error, TEXT("NVENC not avaliable"));
-			return false;
-		}
-
-		TSharedRef<FVideoEncoderInputImpl>	Input(StaticCastSharedRef<FVideoEncoderInputImpl>(InInput));
-
-		FrameFormat = InInput->GetFrameFormat();
-		switch (FrameFormat)
-		{
-	#if PLATFORM_WINDOWS
-		case AVEncoder::EVideoFrameFormat::D3D11_R8G8B8A8_UNORM:
-		case AVEncoder::EVideoFrameFormat::D3D12_R8G8B8A8_UNORM:
-			EncoderDevice = Input->ForceD3D11InputFrames();
-			break;
-	#endif
-	#if WITH_CUDA
-		case AVEncoder::EVideoFrameFormat::CUDA_R8G8B8A8_UNORM:
-			EncoderDevice = Input->GetCUDAEncoderContext();
-			break;
-	#endif
-		case AVEncoder::EVideoFrameFormat::Undefined:
-		default:
-			UE_LOG(LogVideoEncoder, Error, TEXT("Frame format %s is not supported by NVENC_Encoder on this platform."), *ToString(FrameFormat));
-			return false;
-		}
-
-		if (!EncoderDevice)
-		{
-			UE_LOG(LogVideoEncoder, Error, TEXT("NVENC needs encoder device."));
-			return false;
-		}
-
-		if ((MaxFramerate = InInit.MaxFramerate) == 0)
-		{
-			MaxFramerate = 60;
-		}
-		return AddLayer(InInit);
-	}
-
-	FVideoEncoder::FLayerInfo* FVideoEncoderNVENC_H264::CreateLayer(uint32 InLayerIndex, const FLayerInfo& InLayerInfo)
-	{
-		FLayer* Layer = new FLayer(InLayerIndex, InLayerInfo, *this);
-		if (!Layer->Setup())
-		{
-			delete Layer;
-			Layer = nullptr;
-		}
-		return Layer;
-	}
-
-	void FVideoEncoderNVENC_H264::DestroyLayer(FLayerInfo* InLayerInfo)
-	{
-		FLayer* Layer = static_cast<FLayer*>(InLayerInfo);
-		delete Layer;
-	}
-
-	void FVideoEncoderNVENC_H264::Encode(const FVideoEncoderInputFrame* InFrame, const FEncodeOptions& InOptions)
-	{
-		// todo: reconfigure encoder
-		const FVideoEncoderInputFrameImpl* Frame = static_cast<const FVideoEncoderInputFrameImpl*>(InFrame);
-		for (int32 LayerIndex = 0; LayerIndex < LayerInfos.Num(); ++LayerIndex)
-		{
-			FLayer* Layer = static_cast<FLayer*>(LayerInfos[LayerIndex]);
-			Layer->Encode(Frame, InOptions);
-		}
-	}
-
-	void FVideoEncoderNVENC_H264::Flush()
-	{
-		for (int32 LayerIndex = 0; LayerIndex < LayerInfos.Num(); ++LayerIndex)
-		{
-			FLayer* Layer = static_cast<FLayer*>(LayerInfos[LayerIndex]);
-			Layer->Flush();
-		}
-	}
-
-	void FVideoEncoderNVENC_H264::Shutdown()
-	{
-		while (LayerInfos.Num() > 0)
-		{
-			FLayer* Layer = static_cast<FLayer*>(LayerInfos.Pop());
-			Layer->Shutdown();
-			DestroyLayer(Layer);
-		}
-		StopEventThread();
-	}
-
-	void FVideoEncoderNVENC_H264::UpdateFrameRate(uint32 InMaxFramerate)
-	{
-		// UE_LOG(LogVideoEncoder, Error, TEXT("UpdateFrameRate %d->%d"), MaxFramerate, InMaxFramerate);
-		if (MaxFramerate != InMaxFramerate)
-		{
-			MaxFramerate = InMaxFramerate;
-			for (int32 LayerIndex = 0; LayerIndex < LayerInfos.Num(); ++LayerIndex)
-			{
-				FLayer* Layer = static_cast<FLayer*>(LayerInfos[LayerIndex]);
-				Layer->MaxFramerate = MaxFramerate;
-				// Layer->EncoderInitParams.frameRateNum = MaxFramerate;
-				Layer->bUpdateConfig.AtomicSet(true);
-			}
-		}
-	}
-
-	void FVideoEncoderNVENC_H264::UpdateLayerResolution(uint32 InLayerIndex, uint32 InWidth, uint32 InHeight)
-	{
-		if (InLayerIndex < static_cast<uint32>(LayerInfos.Num()))
-		{
-			FLayer* Layer = static_cast<FLayer*>(LayerInfos[InLayerIndex]);
-			Layer->UpdateResolution(InWidth, InHeight);
-			Layer->bUpdateConfig.AtomicSet(true);
-		}
-	}
-
-	void FVideoEncoderNVENC_H264::UpdateLayerBitrate(uint32 InLayerIndex, uint32 InMaxBitRate, uint32 InTargetBitRate)
-	{
-		if (InLayerIndex < static_cast<uint32>(LayerInfos.Num()))
-		{
-			FLayer* Layer = static_cast<FLayer*>(LayerInfos[InLayerIndex]);
-			Layer->UpdateBitrate(InMaxBitRate, InTargetBitRate);
-		}
-	}
-
-	void FVideoEncoderNVENC_H264::UpdateMinQP(int32 minqp)
-	{
-		if (MinQP != minqp)
-		{
-			MinQP = minqp;
-			for (int32 LayerIndex = 0; LayerIndex < LayerInfos.Num(); ++LayerIndex)
-			{
-				FLayer* Layer = static_cast<FLayer*>(LayerInfos[LayerIndex]);
-				Layer->QPMin = MinQP;
-				Layer->bUpdateConfig.AtomicSet(true);
-			}
-		}
-	}
-
-	NV_ENC_PARAMS_RC_MODE ConvertRateControlMode(AVEncoder::FVideoEncoder::RateControlMode mode)
+	NV_ENC_PARAMS_RC_MODE ConvertRateControlModeNVENC(AVEncoder::FVideoEncoder::RateControlMode mode)
 	{
 		switch (mode)
 		{
@@ -218,1202 +32,1267 @@ namespace AVEncoder
 		}
 	}
 
-	void FVideoEncoderNVENC_H264::UpdateRateControl(RateControlMode mode)
+	NV_ENC_MULTI_PASS ConvertMultipassModeNVENC(AVEncoder::FVideoEncoder::MultipassMode mode)
 	{
-		if (RateMode != mode)
+		switch (mode)
 		{
-			RateMode = mode;
-			for (int32 LayerIndex = 0; LayerIndex < LayerInfos.Num(); ++LayerIndex)
-			{
-				FLayer* Layer = static_cast<FLayer*>(LayerInfos[LayerIndex]);
-				Layer->RateControlMode = RateMode;
-				Layer->bUpdateConfig.AtomicSet(true);
-			}
+		case AVEncoder::FVideoEncoder::MultipassMode::DISABLED: return NV_ENC_MULTI_PASS_DISABLED;
+		case AVEncoder::FVideoEncoder::MultipassMode::QUARTER: return NV_ENC_TWO_PASS_QUARTER_RESOLUTION;
+		default:
+		case AVEncoder::FVideoEncoder::MultipassMode::FULL: return NV_ENC_TWO_PASS_FULL_RESOLUTION;
 		}
 	}
+}
 
-	void FVideoEncoderNVENC_H264::UpdateFillData(bool enable)
+namespace AVEncoder
+{
+    static bool GetEncoderInfo(FNVENCCommon& NVENC, FVideoEncoderInfo& EncoderInfo);
+
+    bool FVideoEncoderNVENC_H264::GetIsAvailable(FVideoEncoderInputImpl& InInput, FVideoEncoderInfo& OutEncoderInfo)
+    {
+        FNVENCCommon& NVENC = FNVENCCommon::Setup();
+        bool					bIsAvailable = NVENC.GetIsAvailable();
+        if (bIsAvailable)
+        {
+            OutEncoderInfo.CodecType = ECodecType::H264;
+
+        }
+        return bIsAvailable;
+    }
+
+    void FVideoEncoderNVENC_H264::Register(FVideoEncoderFactory& InFactory)
+    {
+        FNVENCCommon& NVENC = FNVENCCommon::Setup();
+        if (NVENC.GetIsAvailable())
+        {
+            FVideoEncoderInfo	EncoderInfo;
+            if (GetEncoderInfo(NVENC, EncoderInfo))
+            {
+                InFactory.Register(EncoderInfo, []() {
+                    return TUniquePtr<FVideoEncoder>(new FVideoEncoderNVENC_H264());
+                    });
+            }
+        }
+    }
+
+    FVideoEncoderNVENC_H264::FVideoEncoderNVENC_H264()
+        : NVENC(FNVENCCommon::Setup())
+    {
+    }
+
+    FVideoEncoderNVENC_H264::~FVideoEncoderNVENC_H264()
+    {
+        Shutdown();
+    }
+
+    bool FVideoEncoderNVENC_H264::Setup(TSharedRef<FVideoEncoderInput> input, FLayerConfig const& config)
+    {
+        if (!NVENC.GetIsAvailable())
+        {
+            UE_LOG(LogVideoEncoder, Error, TEXT("NVENC not avaliable"));
+            return false;
+        }
+
+        TSharedRef<FVideoEncoderInputImpl>	Input(StaticCastSharedRef<FVideoEncoderInputImpl>(input));
+
+        FrameFormat = input->GetFrameFormat();
+        switch (FrameFormat)
+        {
+    #if PLATFORM_WINDOWS
+        case AVEncoder::EVideoFrameFormat::D3D11_R8G8B8A8_UNORM:
+        case AVEncoder::EVideoFrameFormat::D3D12_R8G8B8A8_UNORM:
+            EncoderDevice = Input->ForceD3D11InputFrames();
+            break;
+    #endif
+    #if WITH_CUDA
+        case AVEncoder::EVideoFrameFormat::CUDA_R8G8B8A8_UNORM:
+            EncoderDevice = Input->GetCUDAEncoderContext();
+            break;
+    #endif
+        case AVEncoder::EVideoFrameFormat::Undefined:
+        default:
+            UE_LOG(LogVideoEncoder, Error, TEXT("Frame format %s is not supported by NVENC_Encoder on this platform."), *ToString(FrameFormat));
+            return false;
+        }
+
+        if (!EncoderDevice)
+        {
+            UE_LOG(LogVideoEncoder, Error, TEXT("NVENC needs encoder device."));
+            return false;
+        }
+
+        auto mutableConfig = config;
+        if (mutableConfig.MaxFramerate == 0)
+            mutableConfig.MaxFramerate = 60;
+
+        return AddLayer(mutableConfig);
+    }
+
+    FVideoEncoder::FLayer* FVideoEncoderNVENC_H264::CreateLayer(uint32 layerIdx, FLayerConfig const& config)
+    {
+        auto const layer = new FNVENCLayer(layerIdx, config, *this);
+        if (!layer->Setup())
+        {
+            delete layer;
+            return nullptr;
+        }
+        return layer;
+    }
+
+	void FVideoEncoderNVENC_H264::DestroyLayer(FLayer* layer)
 	{
-		if (FillData != enable)
-		{
-			FillData = enable;
-			for (int32 LayerIndex = 0; LayerIndex < LayerInfos.Num(); ++LayerIndex)
-			{
-				FLayer* Layer = static_cast<FLayer*>(LayerInfos[LayerIndex]);
-				Layer->FillData = FillData;
-				Layer->bUpdateConfig.AtomicSet(true);
-			}
-		}
+		delete layer;
 	}
 
-	void FVideoEncoderNVENC_H264::OnEvent(void* InEvent, TUniqueFunction<void()>&& InCallback)
-	{
+    void FVideoEncoderNVENC_H264::Encode(FVideoEncoderInputFrame const* frame, FEncodeOptions const& options)
+    {
+        // todo: reconfigure encoder
+        auto const nvencFrame = static_cast<FVideoEncoderInputFrameImpl const*>(frame);
+        for (auto&& layer : Layers)
+        {
+            auto const nvencLayer = static_cast<FNVENCLayer*>(layer);
+            nvencLayer->Encode(nvencFrame, options);
+        }
+    }
+
+    void FVideoEncoderNVENC_H264::Flush()
+    {
+        for (auto&& layer : Layers)
+        {
+            auto const nvencLayer = static_cast<FNVENCLayer*>(layer);
+            nvencLayer->Flush();
+        }
+    }
+
+    void FVideoEncoderNVENC_H264::Shutdown()
+    {
+        for (auto&& layer : Layers)
+        {
+            auto const nvencLayer = static_cast<FNVENCLayer*>(layer);
+            nvencLayer->Shutdown();
+            DestroyLayer(nvencLayer);
+        }
+		Layers.Reset();
+        StopEventThread();
+    }
+
+    void FVideoEncoderNVENC_H264::OnEvent(void* InEvent, TUniqueFunction<void()>&& InCallback)
+    {
 #if PLATFORM_WINDOWS
-		FScopeLock		Guard(&ProtectEventThread);
-		StartEventThread();
-		EventThreadWaitingFor.Emplace(FWaitForEvent(InEvent, MoveTemp(InCallback)));
-		SetEvent(EventThreadCheckEvent);
+        FScopeLock Guard(&ProtectEventThread);
+        StartEventThread();
+        EventThreadWaitingFor.Emplace(FWaitForEvent(InEvent, MoveTemp(InCallback)));
+        SetEvent(EventThreadCheckEvent);
 #else
-		UE_LOG(LogVideoEncoder, Fatal, TEXT("FVideoEncoderNVENC_H264::OnEvent should not be called as NVENC async mode only works on Windows!"))
+        UE_LOG(LogVideoEncoder, Fatal, TEXT("FVideoEncoderNVENC_H264::OnEvent should not be called as NVENC async mode only works on Windows!"))
 #endif
-	}
+    }
 
 // TODO (M84FIX) de-windows this
-	void FVideoEncoderNVENC_H264::StartEventThread()
-	{
+    void FVideoEncoderNVENC_H264::StartEventThread()
+    {
 #if PLATFORM_WINDOWS
-		bExitEventThread = false;
-		if (!EventThread)
-		{
-			if (!EventThreadCheckEvent)
-			{
-				EventThreadCheckEvent = CreateEvent(nullptr, false, false, nullptr);
-				//EventThreadCheckEvent = FPlatformProcess::GetSynchEventFromPool(false);
-			}
-			EventThread = MakeUnique<FThread>(TEXT("NVENC_EncoderCommon"), [this]() { EventLoop(); });
-		}
+        bExitEventThread = false;
+        if (!EventThread)
+        {
+            if (!EventThreadCheckEvent)
+            {
+                EventThreadCheckEvent = CreateEvent(nullptr, false, false, nullptr);
+                //EventThreadCheckEvent = FPlatformProcess::GetSynchEventFromPool(false);
+            }
+            EventThread = MakeUnique<FThread>(TEXT("NVENC_EncoderCommon"), [this]() { EventLoop(); });
+        }
 #else
-		UE_LOG(LogVideoEncoder, Fatal, TEXT("FVideoEncoderNVENC_H264::StartEventThread should not be called as NVENC async mode only works on Windows!"))
+        UE_LOG(LogVideoEncoder, Fatal, TEXT("FVideoEncoderNVENC_H264::StartEventThread should not be called as NVENC async mode only works on Windows!"))
 #endif
-	}
+    }
 
-	void FVideoEncoderNVENC_H264::StopEventThread()
-	{
+    void FVideoEncoderNVENC_H264::StopEventThread()
+    {
 #if PLATFORM_WINDOWS
-		FScopeLock				Guard(&ProtectEventThread);
-		TUniquePtr<FThread>		StopThread = MoveTemp(EventThread);
-		if (StopThread)
-		{
-			bExitEventThread = true;
-			SetEvent(EventThreadCheckEvent);
-			Guard.Unlock();
-			StopThread->Join();
-		}
+        FScopeLock Guard(&ProtectEventThread);
+        TUniquePtr<FThread> StopThread = MoveTemp(EventThread);
+        if (StopThread)
+        {
+            bExitEventThread = true;
+            SetEvent(EventThreadCheckEvent);
+            Guard.Unlock();
+            StopThread->Join();
+        }
 #else
-		UE_LOG(LogVideoEncoder, Fatal, TEXT("FVideoEncoderNVENC_H264::StopEventThread should not be called as NVENC async mode only works on Windows!"))
+        UE_LOG(LogVideoEncoder, Fatal, TEXT("FVideoEncoderNVENC_H264::StopEventThread should not be called as NVENC async mode only works on Windows!"))
 #endif
-	}
+    }
 
 #if PLATFORM_WINDOWS
-	void WindowsError(const TCHAR* lpszFunction)
-	{
-		// Retrieve the system error message for the last-error code
-
-		LPVOID lpMsgBuf;
-		LPVOID lpDisplayBuf;
-		DWORD dw = GetLastError();
-
-		FormatMessage(
-			FORMAT_MESSAGE_ALLOCATE_BUFFER |
-			FORMAT_MESSAGE_FROM_SYSTEM |
-			FORMAT_MESSAGE_IGNORE_INSERTS,
-			NULL,
-			dw,
-			MAKELANGID(LANG_NEUTRAL, SUBLANG_DEFAULT),
-			(LPTSTR)&lpMsgBuf,
-			0, NULL);
-
-		lpDisplayBuf = (LPVOID)LocalAlloc(LMEM_ZEROINIT, (lstrlen((LPCTSTR)lpMsgBuf) + lstrlen((LPCTSTR)lpszFunction) + 40) * sizeof(TCHAR));
-		UE_LOG(LogVideoEncoder, Error, TEXT("%s failed with error %d: %s"), lpszFunction, dw, lpMsgBuf);
-
-		LocalFree(lpMsgBuf);
-		LocalFree(lpDisplayBuf);
-	}
+    void WindowsError(const TCHAR* lpszFunction)
+    {
+        // Retrieve the system error message for the last-error code
+        LPVOID lpMsgBuf;
+        DWORD dw = GetLastError();
+
+        FormatMessage(
+            FORMAT_MESSAGE_ALLOCATE_BUFFER |
+            FORMAT_MESSAGE_FROM_SYSTEM |
+            FORMAT_MESSAGE_IGNORE_INSERTS,
+            NULL,
+            dw,
+            MAKELANGID(LANG_NEUTRAL, SUBLANG_DEFAULT),
+            (LPTSTR)&lpMsgBuf,
+            0, NULL);
+
+        UE_LOG(LogVideoEncoder, Error, TEXT("%s failed with error %d: %s"), lpszFunction, dw, lpMsgBuf);
+
+        LocalFree(lpMsgBuf);
+    }
 #endif
 
-	void FVideoEncoderNVENC_H264::EventLoop()
-	{
+    void FVideoEncoderNVENC_H264::EventLoop()
+    {
 #if PLATFORM_WINDOWS
-		while (true)
-		{
-			// Make an empty array of events to wait for
-			void* EventsToWaitFor[MAXIMUM_WAIT_OBJECTS];
-			int				NumEventsToWaitFor = 0;
-
-			{
-				// Guard the thread and see if we should terminate
-				FScopeLock		Guard(&ProtectEventThread);
-				if (bExitEventThread)
-				{
-					break;
-				}
-
-				// collect events to wait for
-				EventsToWaitFor[NumEventsToWaitFor++] = EventThreadCheckEvent;
-				for (const FWaitForEvent& WaitFor : EventThreadWaitingFor)
-				{
-					check(NumEventsToWaitFor < MAXIMUM_WAIT_OBJECTS);
-					EventsToWaitFor[NumEventsToWaitFor++] = WaitFor.Key;
-				}
-			}
-			// wait for the events
-			DWORD	WaitResult = WaitForMultipleObjects(NumEventsToWaitFor, EventsToWaitFor, false, INFINITE);
-			if (WaitResult >= WAIT_OBJECT_0 && WaitResult < (WAIT_OBJECT_0 + NumEventsToWaitFor))
-			{
-				// Guard the thread and get the event that was triggered
-				FScopeLock		Guard(&ProtectEventThread);
-				void* EventTriggered = EventsToWaitFor[WaitResult - WAIT_OBJECT_0];
-
-				// loop through the the events then callback on the one that was triggered (could this be better as a)
-				for (int32 Index = 0; Index < EventThreadWaitingFor.Num(); ++Index)
-				{
-					if (EventThreadWaitingFor[Index].Key == EventTriggered)
-					{
-						TUniqueFunction<void()>	Callback = MoveTemp(EventThreadWaitingFor[Index].Value);
-						EventThreadWaitingFor.RemoveAtSwap(Index);
-						Guard.Unlock();
-						Callback();
-						break;
-					}
-				}
-			}
-			else if (WaitResult >= WAIT_ABANDONED_0 && WaitResult < (WAIT_ABANDONED_0 + NumEventsToWaitFor))
-			{
-			}
-			else if (WaitResult == WAIT_TIMEOUT)
-			{
-			}
-			else if (WaitResult == WAIT_FAILED)
-			{
-				// HACK to fix warning in clang
-				WindowsError(TEXT("WaitForMultipleObjects"));
-			}
-		}
+        while (true)
+        {
+            // Make an empty array of events to wait for
+            void* EventsToWaitFor[MAXIMUM_WAIT_OBJECTS];
+            int				NumEventsToWaitFor = 0;
+
+            {
+                // Guard the thread and see if we should terminate
+                FScopeLock		Guard(&ProtectEventThread);
+                if (bExitEventThread)
+                {
+                    break;
+                }
+
+                // collect events to wait for
+                EventsToWaitFor[NumEventsToWaitFor++] = EventThreadCheckEvent;
+                for (const FWaitForEvent& WaitFor : EventThreadWaitingFor)
+                {
+                    check(NumEventsToWaitFor < MAXIMUM_WAIT_OBJECTS);
+                    EventsToWaitFor[NumEventsToWaitFor++] = WaitFor.Key;
+                }
+            }
+            // wait for the events
+            DWORD	WaitResult = WaitForMultipleObjects(NumEventsToWaitFor, EventsToWaitFor, false, INFINITE);
+            if (WaitResult >= WAIT_OBJECT_0 && WaitResult < (WAIT_OBJECT_0 + NumEventsToWaitFor))
+            {
+                // Guard the thread and get the event that was triggered
+                FScopeLock		Guard(&ProtectEventThread);
+                void* EventTriggered = EventsToWaitFor[WaitResult - WAIT_OBJECT_0];
+
+                // loop through the the events then callback on the one that was triggered (could this be better as a)
+                for (int32 Index = 0; Index < EventThreadWaitingFor.Num(); ++Index)
+                {
+                    if (EventThreadWaitingFor[Index].Key == EventTriggered)
+                    {
+                        TUniqueFunction<void()>	Callback = MoveTemp(EventThreadWaitingFor[Index].Value);
+                        EventThreadWaitingFor.RemoveAtSwap(Index);
+                        Guard.Unlock();
+                        Callback();
+                        break;
+                    }
+                }
+            }
+            else if (WaitResult >= WAIT_ABANDONED_0 && WaitResult < (WAIT_ABANDONED_0 + NumEventsToWaitFor))
+            {
+            }
+            else if (WaitResult == WAIT_TIMEOUT)
+            {
+            }
+            else if (WaitResult == WAIT_FAILED)
+            {
+                // HACK to fix warning in clang
+                WindowsError(TEXT("WaitForMultipleObjects"));
+            }
+        }
 #else
-		UE_LOG(LogVideoEncoder, Fatal, TEXT("FVideoEncoderNVENC_H264::EventLoop should not be called as NVENC async mode only works on Windows!"))
+        UE_LOG(LogVideoEncoder, Fatal, TEXT("FVideoEncoderNVENC_H264::EventLoop should not be called as NVENC async mode only works on Windows!"))
 #endif
-	}
-
-	// --- FVideoEncoderNVENC_H264::FLayer ------------------------------------------------------------
-	FVideoEncoderNVENC_H264::FLayer::FLayer(uint32 InLayerIndex, const FLayerInfo& InLayerInfo, FVideoEncoderNVENC_H264& InEncoder)
-		: FLayerInfo(InLayerInfo)
-		, Encoder(InEncoder)
-		, NVENC(FNVENCCommon::Setup())
-		, CodecGUID(NV_ENC_CODEC_H264_GUID)
-		, LayerIndex(InLayerIndex)
-	{
-	}
-
-	FVideoEncoderNVENC_H264::FLayer::~FLayer()
-	{
-	}
-
-	bool FVideoEncoderNVENC_H264::FLayer::Setup()
-	{
-		bool bSuccess = false;
-
-		if (CreateSession() && CreateInitialConfig())
-		{
-			// create encoder
-			NVENCSTATUS		Result = NVENC.nvEncInitializeEncoder(NVEncoder, &EncoderInitParams);
-			if (Result != NV_ENC_SUCCESS)
+    }
+
+    // --- FVideoEncoderNVENC_H264::FLayer ------------------------------------------------------------
+    FVideoEncoderNVENC_H264::FNVENCLayer::FNVENCLayer(uint32 layerIdx, FLayerConfig const& config, FVideoEncoderNVENC_H264& encoder)
+        : FLayer(config)
+        , Encoder(encoder)
+        , NVENC(FNVENCCommon::Setup())
+        , CodecGUID(NV_ENC_CODEC_H264_GUID)
+        , LayerIndex(layerIdx)
+    {
+    }
+
+    FVideoEncoderNVENC_H264::FNVENCLayer::~FNVENCLayer()
+    {
+    }
+
+    bool FVideoEncoderNVENC_H264::FNVENCLayer::Setup()
+    {
+        if (CreateSession() && CreateInitialConfig())
+        {
+            // create encoder
+			auto const result = NVENC.nvEncInitializeEncoder(NVEncoder, &EncoderInitParams);
+			if (result != NV_ENC_SUCCESS)
 			{
-				UE_LOG(LogVideoEncoder, Error, TEXT("Unable to initialize NvEnc encoder (%s)."), *NVENC.GetErrorString(NVEncoder, Result));
+				UE_LOG(LogVideoEncoder, Error, TEXT("Unable to initialize NvEnc encoder (%s)."), *NVENC.GetErrorString(NVEncoder, result));
 			}
 			else
-			{
-				bSuccess = true;
-			}
-		}
-
-		return bSuccess;
-	}
-
-	bool FVideoEncoderNVENC_H264::FLayer::CreateSession()
-	{
-		if (!NVEncoder)
-		{
-			// create the encoder session
-			NVENCStruct(NV_ENC_OPEN_ENCODE_SESSION_EX_PARAMS, OpenEncodeSessionExParams);
-			OpenEncodeSessionExParams.apiVersion = NVENCAPI_VERSION;
-			OpenEncodeSessionExParams.device = Encoder.EncoderDevice;
-
-			switch (Encoder.FrameFormat)
-			{
-			case EVideoFrameFormat::D3D11_R8G8B8A8_UNORM:
-			case EVideoFrameFormat::D3D12_R8G8B8A8_UNORM:
-				OpenEncodeSessionExParams.deviceType = NV_ENC_DEVICE_TYPE_DIRECTX;
-				break;
-			case EVideoFrameFormat::CUDA_R8G8B8A8_UNORM:
-				OpenEncodeSessionExParams.deviceType = NV_ENC_DEVICE_TYPE_CUDA;
-				break;
-			default:
-				UE_LOG(LogVideoEncoder, Error, TEXT("FrameFormat %s unavailable."), *ToString(Encoder.FrameFormat));
-				return false;
-				break;
-			}
-
-			NVENCSTATUS		NvResult = NVENC.nvEncOpenEncodeSessionEx(&OpenEncodeSessionExParams, &NVEncoder);
-			if (NvResult != NV_ENC_SUCCESS)
-			{
-				UE_LOG(LogVideoEncoder, Error, TEXT("Unable to open NvEnc encoding session (%s)."), *NVENC.GetErrorString(NVEncoder, NvResult));
+				return true;
+        }
+
+        return false;
+    }
+
+    bool FVideoEncoderNVENC_H264::FNVENCLayer::CreateSession()
+    {
+        if (!NVEncoder)
+        {
+            // create the encoder session
+            NVENCStruct(NV_ENC_OPEN_ENCODE_SESSION_EX_PARAMS, OpenEncodeSessionExParams);
+            OpenEncodeSessionExParams.apiVersion = NVENCAPI_VERSION;
+            OpenEncodeSessionExParams.device = Encoder.EncoderDevice;
+
+            switch (Encoder.FrameFormat)
+            {
+            case EVideoFrameFormat::D3D11_R8G8B8A8_UNORM:
+            case EVideoFrameFormat::D3D12_R8G8B8A8_UNORM:
+                OpenEncodeSessionExParams.deviceType = NV_ENC_DEVICE_TYPE_DIRECTX;
+                break;
+            case EVideoFrameFormat::CUDA_R8G8B8A8_UNORM:
+                OpenEncodeSessionExParams.deviceType = NV_ENC_DEVICE_TYPE_CUDA;
+                break;
+            default:
+                UE_LOG(LogVideoEncoder, Error, TEXT("FrameFormat %s unavailable."), *ToString(Encoder.FrameFormat));
+                return false;
+                break;
+            }
+
+			auto const result = NVENC.nvEncOpenEncodeSessionEx(&OpenEncodeSessionExParams, &NVEncoder);
+            if (result != NV_ENC_SUCCESS)
+            {
+                UE_LOG(LogVideoEncoder, Error, TEXT("Unable to open NvEnc encoding session (%s)."), *NVENC.GetErrorString(NVEncoder, result));
 				NVEncoder = nullptr;
-			}
-		}
-		return NVEncoder != nullptr;
-	}
-
-
-	// TODO (M84FIX) need to parameterize this
-	bool FVideoEncoderNVENC_H264::FLayer::CreateInitialConfig()
-	{
-		NVENCSTATUS		Result;
-
-		// set the initialization parameters
-		FMemory::Memzero(EncoderInitParams);
-		EncoderInitParams.version = NV_ENC_INITIALIZE_PARAMS_VER;
-		EncoderInitParams.encodeWidth = EncoderInitParams.darWidth = Width;
-		EncoderInitParams.encodeHeight = EncoderInitParams.darHeight = Height;
-		EncoderInitParams.encodeGUID = NV_ENC_CODEC_H264_GUID;
-
-		EncoderInitParams.encodeGUID = NV_ENC_CODEC_H264_GUID;
-		EncoderInitParams.presetGUID = NV_ENC_PRESET_P4_GUID;
-
-		EncoderInitParams.frameRateNum = Encoder.MaxFramerate;
-		EncoderInitParams.frameRateDen = 1;
-		EncoderInitParams.enablePTD = 1;
-		EncoderInitParams.reportSliceOffsets = 0;
-		EncoderInitParams.enableSubFrameWrite = 0;
-		EncoderInitParams.encodeConfig = &EncoderConfig;
-		EncoderInitParams.maxEncodeWidth = 4096;
-		EncoderInitParams.maxEncodeHeight = 4096;
-		EncoderInitParams.tuningInfo = NV_ENC_TUNING_INFO_LOW_LATENCY;
-
-		// load a preset configuration
-		NVENCStruct(NV_ENC_PRESET_CONFIG, PresetConfig);
-		PresetConfig.presetCfg.version = NV_ENC_CONFIG_VER;
-
-		Result = NVENC.nvEncGetEncodePresetConfigEx(NVEncoder, EncoderInitParams.encodeGUID, EncoderInitParams.presetGUID, EncoderInitParams.tuningInfo, &PresetConfig);
-		if (Result != NV_ENC_SUCCESS)
-		{
-			UE_LOG(LogVideoEncoder, Error, TEXT("Unable to get NvEnc preset config (%s)."), *NVENC.GetErrorString(NVEncoder, Result));
-			return false;
-		}
-		
-		FMemory::Memcpy(&EncoderConfig, &PresetConfig.presetCfg, sizeof(NV_ENC_CONFIG));
-
+				return false;
+            }
+        }
+
+        return true;
+    }
+
+    bool FVideoEncoderNVENC_H264::FNVENCLayer::CreateInitialConfig()
+    {
+        // set the initialization parameters
+        FMemory::Memzero(EncoderInitParams);
+
+		CurrentConfig.MaxFramerate = 60;
+
+        EncoderInitParams.version = NV_ENC_INITIALIZE_PARAMS_VER;
+        EncoderInitParams.encodeGUID = NV_ENC_CODEC_H264_GUID;
+        EncoderInitParams.presetGUID = NV_ENC_PRESET_P4_GUID;
+        EncoderInitParams.frameRateNum = CurrentConfig.MaxFramerate;
+        EncoderInitParams.frameRateDen = 1;
+        EncoderInitParams.enablePTD = 1;
+        EncoderInitParams.reportSliceOffsets = 0;
+        EncoderInitParams.enableSubFrameWrite = 0;
+        EncoderInitParams.maxEncodeWidth = 4096;
+        EncoderInitParams.maxEncodeHeight = 4096;
+        EncoderInitParams.tuningInfo = NV_ENC_TUNING_INFO_ULTRA_LOW_LATENCY;
+
+        // load a preset configuration
+        NVENCStruct(NV_ENC_PRESET_CONFIG, PresetConfig);
+        PresetConfig.presetCfg.version = NV_ENC_CONFIG_VER;
+        auto const result = NVENC.nvEncGetEncodePresetConfigEx(NVEncoder, EncoderInitParams.encodeGUID, EncoderInitParams.presetGUID, EncoderInitParams.tuningInfo, &PresetConfig);
+        if (result != NV_ENC_SUCCESS)
+        {
+            UE_LOG(LogVideoEncoder, Error, TEXT("Unable to get NvEnc preset config (%s)."), *NVENC.GetErrorString(NVEncoder, result));
+            return false;
+        }
+        
+		// copy the preset config to our config
+        FMemory::Memcpy(&EncoderConfig, &PresetConfig.presetCfg, sizeof(NV_ENC_CONFIG));
 		EncoderConfig.profileGUID = NV_ENC_H264_PROFILE_BASELINE_GUID;
+		EncoderConfig.rcParams.version = NV_ENC_RC_PARAMS_VER;
+		EncoderInitParams.encodeConfig = &EncoderConfig;
 
-		EncoderConfig.encodeCodecConfig.hevcConfig.enableFillerDataInsertion = FillData ? 1 : 0;
-		EncoderConfig.encodeCodecConfig.h264Config.enableFillerDataInsertion = FillData ? 1 : 0;
-
-		// set up rate control configuration
-		NV_ENC_RC_PARAMS& RcParams = EncoderConfig.rcParams;
-		RcParams.version = NV_ENC_RC_PARAMS_VER;
-
-		RcParams.rateControlMode = ConvertRateControlMode(RateControlMode);
-
-		RcParams.enableMinQP = QPMin > -1;
-		auto const minqp = static_cast<uint32_t>(QPMin);
-		RcParams.minQP = { minqp, minqp, minqp };
-
-		RcParams.enableMinQP = false;
-		RcParams.enableMaxQP = false;
+        // h264 specific settings
 
-		RcParams.maxBitRate = MaxBitrate;   // Not used for CBR
-		RcParams.averageBitRate = FMath::Min(RcParams.maxBitRate, DEFAULT_BITRATE);
+        // repeat SPS/PPS with each key-frame for a case when the first frame (with mandatory SPS/PPS) 
+        // was dropped by WebRTC
+        EncoderConfig.encodeCodecConfig.h264Config.repeatSPSPPS = 1;
 
-		// h264 specific settings
+        // configure "entire frame as a single slice"
+        // seems WebRTC implementation doesn't work well with slicing, default mode 
+        // produces (rarely, but specially under packet loss) grey full screen or just top half of it.
+        EncoderConfig.encodeCodecConfig.h264Config.sliceMode = 0;
+        EncoderConfig.encodeCodecConfig.h264Config.sliceModeData = 0;
 
-		// repeat SPS/PPS with each key-frame for a case when the first frame (with mandatory SPS/PPS) 
-		// was dropped by WebRTC
-		EncoderConfig.encodeCodecConfig.h264Config.repeatSPSPPS = 1;
+        // whether or not to use async mode
+        if (GetCapability(NV_ENC_CAPS_ASYNC_ENCODE_SUPPORT) != 0)
+        {
+            EncoderInitParams.enableEncodeAsync = true;
+            bAsyncMode = true;
+        }
 
-		// configure "entire frame as a single slice"
-		// seems WebRTC implementation doesn't work well with slicing, default mode 
-		// produces (rarely, but specially under packet loss) grey full screen or just top half of it.
-		EncoderConfig.encodeCodecConfig.h264Config.sliceMode = 0;
-		EncoderConfig.encodeCodecConfig.h264Config.sliceModeData = 0;
+		// update config from CurrentConfig
+		UpdateConfig();
 
-		// whether or not to use async mode
-		if (GetCapability(NV_ENC_CAPS_ASYNC_ENCODE_SUPPORT) != 0)
-		{
-			EncoderInitParams.enableEncodeAsync = true;
-			bAsyncMode = true;
-		}
-		else
-		{
-		}
-		return true;
-	}
+        return true;
+    }
 
-	void FVideoEncoderNVENC_H264::FLayer::Encode(const FVideoEncoderInputFrameImpl* InFrame, const FEncodeOptions& InOptions)
+	void FVideoEncoderNVENC_H264::FNVENCLayer::MaybeReconfigure()
 	{
-		FInputOutput* Buffer = GetOrCreateBuffer(InFrame);
-
-		if (Buffer)
-		{
-			Buffer->EncodeStartTs = FTimespan::FromSeconds(FPlatformTime::Seconds());
-
-			if (bUpdateConfig.AtomicSet(false))
-			{
-				// Changing the framerate forces the generation of a new keyframe so we only do it when either:
-				// 1) we need a keyframe anyway
-				// 2) there is a big difference in framerate
-				// 3) the framerate has changed and we haven't sent a Keyframe recently
-				uint32 CurrentMaxFramerate = EncoderInitParams.frameRateNum;
-				uint32 FrameRateDiff = MaxFramerate > CurrentMaxFramerate ? MaxFramerate - CurrentMaxFramerate : CurrentMaxFramerate - MaxFramerate;
-
-				if (InOptions.bForceKeyFrame ||
-					FrameRateDiff > MAX_FRAMERATE_DIFF ||
-					(MaxFramerate != EncoderInitParams.frameRateNum && (FDateTime::UtcNow() - LastKeyFrameTime).GetSeconds() > MIN_UPDATE_FRAMERATE_SECS))
-				{
-					EncoderInitParams.frameRateNum = MaxFramerate;
-				}
-
-				NVENCStruct(NV_ENC_RECONFIGURE_PARAMS, ReconfigureParams);
-				FMemory::Memcpy(&ReconfigureParams.reInitEncodeParams, &EncoderInitParams, sizeof(EncoderInitParams));
-				// Bitrate generated by NVENC is not very stable when the scene doesnt have a lot of movement
-				// outputPictureTimingSEI enables the filling data when using CBR so that the bitrate generated
-				// is much closer to the requested one and bandwidth estimation algorithms can work better.
-				// Otherwise in a static scene it can send 50kbps when configuring 300kbps and it will never ramp up.
-				if (EncoderInitParams.encodeConfig->rcParams.averageBitRate < 5000000)
-				{
-					EncoderInitParams.encodeConfig->encodeCodecConfig.h264Config.outputPictureTimingSEI = 1;
-					EncoderInitParams.encodeConfig->rcParams.enableMinQP = 0;
-				}
-				else
-				{
-					EncoderInitParams.encodeConfig->encodeCodecConfig.h264Config.outputPictureTimingSEI = 0;
-					EncoderInitParams.encodeConfig->rcParams.enableMinQP = 1;
-					EncoderInitParams.encodeConfig->rcParams.minQP = { 20, 20, 20 };
-				}
-
-				EncoderInitParams.encodeConfig->rcParams.enableMinQP = QPMin > -1;
-				if (QPMin > -1)
-				{
-					auto minqp = static_cast<uint32_t>(QPMin);
-					EncoderInitParams.encodeConfig->rcParams.minQP = { minqp, minqp, minqp };
-				}
-
-				EncoderConfig.rcParams.rateControlMode = ConvertRateControlMode(RateControlMode);
-
-				EncoderInitParams.encodeConfig->encodeCodecConfig.hevcConfig.enableFillerDataInsertion = FillData ? 1 : 0;
-				EncoderInitParams.encodeConfig->encodeCodecConfig.h264Config.enableFillerDataInsertion = FillData ? 1 : 0;
-
-				NVENCSTATUS	Result = NVENC.nvEncReconfigureEncoder(NVEncoder, &ReconfigureParams);
-				UE_LOG(LogVideoEncoder, VeryVerbose, TEXT("NVENC.nvEncReconfigureEncoder(NVEncoder, &ReconfigureParams); -> %d  (%d, %d, %d, %d, %ul)"), Result,
-					EncoderInitParams.encodeConfig->rcParams.averageBitRate, EncoderInitParams.encodeConfig->rcParams.maxBitRate,
-					EncoderInitParams.frameRateNum, MaxFramerate, InFrame->PTS);
-				if (Result != NV_ENC_SUCCESS)
-				{
-					UE_LOG(LogVideoEncoder, Error, TEXT("Failed to update NVENC encoder configuration (%s)"), *NVENC.GetErrorString(NVEncoder, Result));
-				}
-			}
-
-			if (MapInputTexture(*Buffer))
-			{
-				NVENCStruct(NV_ENC_PIC_PARAMS, PicParams);
-				PicParams.inputWidth = Buffer->Width;
-				PicParams.inputHeight = Buffer->Height;
-				PicParams.inputPitch = Buffer->Pitch;
-				PicParams.inputBuffer = Buffer->MappedInput;
-				PicParams.bufferFmt = Buffer->BufferFormat;
-				PicParams.encodePicFlags = 0;
-				if (InOptions.bForceKeyFrame)
-				{
-					LastKeyFrameTime = FDateTime::UtcNow();
-					PicParams.encodePicFlags |= NV_ENC_PIC_FLAG_FORCEIDR;
-				}
-				PicParams.inputTimeStamp = Buffer->TimeStamp = InFrame->PTS;
-				PicParams.outputBitstream = Buffer->OutputBitstream;
-				PicParams.completionEvent = Buffer->CompletionEvent;
-				PicParams.pictureStruct = NV_ENC_PIC_STRUCT_FRAME;
-				NVENCSTATUS		Result = NVENC.nvEncEncodePicture(NVEncoder, &PicParams);
-				//			UE_LOG(LogVideoEncoder, Error, TEXT("NVENC.nvEncEncodePicture(NVEncoder, &PicParams); -> %d"), Result);
-
-				if (Result == NV_ENC_ERR_NEED_MORE_INPUT)
-				{
-					// queue to pending frames to be read on next success
-					PendingEncodes.Enqueue(Buffer);
-				}
-				else if (Result == NV_ENC_SUCCESS)
-				{
-					PendingEncodes.Enqueue(Buffer);
-					WaitForNextPendingFrame();
-				}
-				else
-				{
-					UE_LOG(LogVideoEncoder, Error, TEXT("NVENC.nvEncEncodePicture(NVEncoder, &PicParams); -> %s"), *NVENC.GetErrorString(NVEncoder, Result));
-					// release input frame
-					if (Buffer->SourceFrame)
-					{
-						Buffer->SourceFrame->Release();
-						Buffer->SourceFrame = nullptr;
-					}
-				}
-			}
-			else
+		FScopeLock lock(&ConfigMutex);
+		if (NeedsReconfigure)
+		{
+			UpdateConfig();
+
+			// Changing the framerate forces the generation of a new keyframe so we only do it when either:
+			// 1) we need a keyframe anyway
+			// 2) there is a big difference in framerate
+			// 3) the framerate has changed and we haven't sent a Keyframe recently
+			auto const CurrentMaxFramerate = EncoderInitParams.frameRateNum;
+			auto const ConfigMaxFramerate = CurrentConfig.MaxFramerate;
+			auto const FrameRateDiff = ConfigMaxFramerate > CurrentMaxFramerate ? ConfigMaxFramerate - CurrentMaxFramerate : CurrentMaxFramerate - ConfigMaxFramerate;
+			auto const LastKeyFrameDelta = (FDateTime::UtcNow() - LastKeyFrameTime).GetSeconds();
+
+			if (bForceNextKeyframe ||
+				FrameRateDiff > MAX_FRAMERATE_DIFF ||
+				(ConfigMaxFramerate != EncoderInitParams.frameRateNum && LastKeyFrameDelta > MIN_UPDATE_FRAMERATE_SECS))
 			{
-				// todo: release source frame
+				EncoderInitParams.frameRateNum = CurrentConfig.MaxFramerate;
 			}
-		}
-	}
-
-	void FVideoEncoderNVENC_H264::FLayer::Flush()
-	{
-		FInputOutput* EmptyBuffer = CreateBuffer();
-		if (!EmptyBuffer)
-		{
-			return;
-		}
-		NVENCStruct(NV_ENC_PIC_PARAMS, PicParams);
-		PicParams.encodePicFlags = NV_ENC_PIC_FLAG_EOS;
-		PicParams.completionEvent = EmptyBuffer->CompletionEvent;
-		NVENCSTATUS		Result = NVENC.nvEncEncodePicture(NVEncoder, &PicParams);
-		//	UE_LOG(LogVideoEncoder, Error, TEXT("NVENC.nvEncEncodePicture(NVEncoder, &PicParams); (flush) -> %d"), Result);
-		if (Result != NV_ENC_SUCCESS)
-		{
-			UE_LOG(LogVideoEncoder, Warning, TEXT("Failed to flush NVENC encoder (%s)"), *NVENC.GetErrorString(NVEncoder, Result));
-		}
-		else
-		{
-			EmptyBuffer->TriggerOnCompletion = FPlatformProcess::GetSynchEventFromPool(true);
+			bForceNextKeyframe = false;
 
-			PendingEncodes.Enqueue(EmptyBuffer);
-			WaitForNextPendingFrame();
-			// wait for this buffer to be done
-			EmptyBuffer->TriggerOnCompletion->Wait();
-
-			DestroyBuffer(EmptyBuffer);
-			for (FInputOutput* Buffer : CreatedBuffers)
-			{
-				DestroyBuffer(Buffer);
-			}
-			CreatedBuffers.Empty();
-		}
-	}
+            // `outputPictureTimingSEI` is used in CBR mode to fill video frame with data to match the requested bitrate.
+            if(CurrentConfig.RateControlMode == AVEncoder::FVideoEncoder::RateControlMode::CBR)
+            {
+                EncoderInitParams.encodeConfig->encodeCodecConfig.h264Config.outputPictureTimingSEI = EncoderInitParams.encodeConfig->rcParams.enableMinQP ? 0 : 1;
+            }
 
-	void FVideoEncoderNVENC_H264::FLayer::Shutdown()
-	{
-		Flush();
-		if (NVEncoder)
-		{
-			NVENCSTATUS		Result = NVENC.nvEncDestroyEncoder(NVEncoder);
-			//		UE_LOG(LogVideoEncoder, Error, TEXT("NVENC.nvEncDestroyEncoder(NVEncoder); -> %d"), Result);
-			if (Result != NV_ENC_SUCCESS)
-			{
-				UE_LOG(LogVideoEncoder, Error, TEXT("Failed to destroy NVENC encoder (%s)."), *NVENC.GetErrorString(NVEncoder, Result));
-			}
-			NVEncoder = nullptr;
-			bAsyncMode = false;
-		}
-	}
+			NVENCStruct(NV_ENC_RECONFIGURE_PARAMS, ReconfigureParams);
+			FMemory::Memcpy(&ReconfigureParams.reInitEncodeParams, &EncoderInitParams, sizeof(EncoderInitParams));
 
-	void FVideoEncoderNVENC_H264::FLayer::UpdateBitrate(uint32 InMaxBitRate, uint32 InTargetBitRate)
-	{
-		NV_ENC_RC_PARAMS& RcParams = EncoderConfig.rcParams;
-		// UE_LOG(LogVideoEncoder, Error, TEXT("UpdateBitrate Avg %d->%d Max %d->%d"), RcParams.averageBitRate, InTargetBitRate, RcParams.maxBitRate, InMaxBitRate);
+			auto const result = NVENC.nvEncReconfigureEncoder(NVEncoder, &ReconfigureParams);
+			if (result != NV_ENC_SUCCESS)
+				UE_LOG(LogVideoEncoder, Error, TEXT("Failed to update NVENC encoder configuration (%s)"), *NVENC.GetErrorString(NVEncoder, result));
 
-		if (RcParams.averageBitRate != InTargetBitRate || RcParams.maxBitRate != InMaxBitRate)
-		{
-			RcParams.averageBitRate = InTargetBitRate;
-			RcParams.maxBitRate = InMaxBitRate;
-			bUpdateConfig.AtomicSet(true);
+			NeedsReconfigure = false;
 		}
 	}
 
-	void FVideoEncoderNVENC_H264::FLayer::UpdateResolution(uint32 InWidth, uint32 InHeight)
+	void FVideoEncoderNVENC_H264::FNVENCLayer::UpdateConfig()
 	{
-		this->Width = InWidth;
-		EncoderInitParams.encodeWidth = InWidth;
-		EncoderInitParams.darWidth = InWidth;
-
-		this->Height = InHeight;
-		EncoderInitParams.encodeHeight = InHeight;
-		EncoderInitParams.darHeight = InHeight;
-
-		bUpdateConfig.AtomicSet(true);
+		EncoderInitParams.encodeWidth = EncoderInitParams.darWidth = CurrentConfig.Width;
+		EncoderInitParams.encodeHeight = EncoderInitParams.darHeight = CurrentConfig.Height;
+
+		auto& rateContolParams = EncoderInitParams.encodeConfig->rcParams;
+		rateContolParams.rateControlMode = ConvertRateControlModeNVENC(CurrentConfig.RateControlMode);
+		rateContolParams.averageBitRate = CurrentConfig.TargetBitrate > -1 ? CurrentConfig.TargetBitrate : DEFAULT_BITRATE;
+		rateContolParams.maxBitRate = CurrentConfig.MaxBitrate > -1 ? CurrentConfig.MaxBitrate : DEFAULT_BITRATE;   // Not used for CBR
+		rateContolParams.multiPass = ConvertMultipassModeNVENC(CurrentConfig.MultipassMode);
+		auto const minqp = static_cast<uint32_t>(CurrentConfig.QPMin);
+		auto const maxqp = static_cast<uint32_t>(CurrentConfig.QPMax);
+		rateContolParams.minQP = { minqp, minqp, minqp };
+		rateContolParams.maxQP = { maxqp, maxqp, maxqp };
+		rateContolParams.enableMinQP = CurrentConfig.QPMin > -1;
+		rateContolParams.enableMaxQP = CurrentConfig.QPMax > -1;
+
+		auto& h264Config = EncoderInitParams.encodeConfig->encodeCodecConfig.h264Config;
+		h264Config.enableFillerDataInsertion = CurrentConfig.FillData ? 1 : 0;
 	}
 
-	void FVideoEncoderNVENC_H264::FLayer::ProcessNextPendingFrame()
-	{
-		FInputOutput* Buffer = nullptr;
-		if (PendingEncodes.Dequeue(Buffer))
-		{
-			// this could be a null buffer (for flush)
-			if (Buffer->Width > 0 && Buffer->Height > 0)
-			{
-				// lock output buffers for CPU access
-				if (LockOutputBuffer(*Buffer))
-				{
-					if (Encoder.OnEncodedPacket)
-					{
-						// create packet with buffer contents
-						FCodecPacketImpl	Packet;
-
-						Packet.PTS = static_cast<int64>(Buffer->TimeStamp);
-						Packet.Data = static_cast<const uint8*>(Buffer->BitstreamData);
-						Packet.DataSize = Buffer->BitstreamDataSize;
-						if (Buffer->PictureType == NV_ENC_PIC_TYPE_IDR)
-						{
-							UE_LOG(LogVideoEncoder, Verbose, TEXT("Generated IDR Frame"));
-							Packet.IsKeyFrame = true;
-						}
-						Packet.VideoQP = Buffer->FrameAvgQP;
-						Packet.Timings.StartTs = Buffer->EncodeStartTs;
-						Packet.Timings.FinishTs = FTimespan::FromSeconds(FPlatformTime::Seconds());
-						Packet.Framerate = MaxFramerate;
-
-						Encoder.OnEncodedPacket(LayerIndex, Buffer->SourceFrame, Packet);
-					}
-
-					UnlockOutputBuffer(*Buffer);
-				}
-			}
-
-			// input is no longer required
-			if (Buffer->SourceFrame)
-			{
-				Buffer->SourceFrame->Release();
-				Buffer->SourceFrame = nullptr;
-			}
-
-			// optional event to be triggered
-			if (Buffer->TriggerOnCompletion)
+    void FVideoEncoderNVENC_H264::FNVENCLayer::Encode(FVideoEncoderInputFrameImpl const* frame, FEncodeOptions const& options)
+    {
+        FInputOutput* Buffer = GetOrCreateBuffer(frame);
+
+        if (Buffer)
+        {
+            Buffer->EncodeStartTs = FTimespan::FromSeconds(FPlatformTime::Seconds());
+
+			bForceNextKeyframe = options.bForceKeyFrame;
+			MaybeReconfigure();
+
+            if (MapInputTexture(*Buffer))
+            {
+                NVENCStruct(NV_ENC_PIC_PARAMS, PicParams);
+                PicParams.inputWidth = Buffer->Width;
+                PicParams.inputHeight = Buffer->Height;
+                PicParams.inputPitch = Buffer->Pitch;
+                PicParams.inputBuffer = Buffer->MappedInput;
+                PicParams.bufferFmt = Buffer->BufferFormat;
+                PicParams.encodePicFlags = 0;
+                if (options.bForceKeyFrame)
+                {
+                    LastKeyFrameTime = FDateTime::UtcNow();
+                    PicParams.encodePicFlags |= NV_ENC_PIC_FLAG_FORCEIDR;
+                }
+                PicParams.inputTimeStamp = Buffer->TimeStamp = frame->PTS;
+                PicParams.outputBitstream = Buffer->OutputBitstream;
+                PicParams.completionEvent = Buffer->CompletionEvent;
+                PicParams.pictureStruct = NV_ENC_PIC_STRUCT_FRAME;
+
+                auto const result = NVENC.nvEncEncodePicture(NVEncoder, &PicParams);
+                if (result == NV_ENC_ERR_NEED_MORE_INPUT)
+                {
+                    // queue to pending frames to be read on next success
+                    PendingEncodes.Enqueue(Buffer);
+                }
+                else if (result == NV_ENC_SUCCESS)
+                {
+                    PendingEncodes.Enqueue(Buffer);
+                    WaitForNextPendingFrame();
+                }
+                else
+                {
+                    UE_LOG(LogVideoEncoder, Error, TEXT("NVENC.nvEncEncodePicture(NVEncoder, &PicParams); -> %s"), *NVENC.GetErrorString(NVEncoder, result));
+                    // release input frame
+                    if (Buffer->SourceFrame)
+                    {
+                        Buffer->SourceFrame->Release();
+                        Buffer->SourceFrame = nullptr;
+                    }
+                }
+            }
+            else
+            {
+                // todo: release source frame
+            }
+        }
+    }
+
+    void FVideoEncoderNVENC_H264::FNVENCLayer::Flush()
+    {
+        FInputOutput* EmptyBuffer = CreateBuffer();
+        if (!EmptyBuffer)
+        {
+            return;
+        }
+        NVENCStruct(NV_ENC_PIC_PARAMS, PicParams);
+        PicParams.encodePicFlags = NV_ENC_PIC_FLAG_EOS;
+        PicParams.completionEvent = EmptyBuffer->CompletionEvent;
+        auto const result = NVENC.nvEncEncodePicture(NVEncoder, &PicParams);
+		if (result != NV_ENC_SUCCESS)
+		{
+			UE_LOG(LogVideoEncoder, Warning, TEXT("Failed to flush NVENC encoder (%s)"), *NVENC.GetErrorString(NVEncoder, result));
+		}
+        else
+        {
+            EmptyBuffer->TriggerOnCompletion = FPlatformProcess::GetSynchEventFromPool(true);
+
+            PendingEncodes.Enqueue(EmptyBuffer);
+            WaitForNextPendingFrame();
+            // wait for this buffer to be done
+            EmptyBuffer->TriggerOnCompletion->Wait();
+
+            DestroyBuffer(EmptyBuffer);
+            for (FInputOutput* Buffer : CreatedBuffers)
+                DestroyBuffer(Buffer);
+            CreatedBuffers.Empty();
+        }
+    }
+
+    void FVideoEncoderNVENC_H264::FNVENCLayer::Shutdown()
+    {
+        Flush();
+        if (NVEncoder)
+        {
+            auto const result = NVENC.nvEncDestroyEncoder(NVEncoder);
+			if (result != NV_ENC_SUCCESS)
 			{
-				Buffer->TriggerOnCompletion->Trigger();
+				UE_LOG(LogVideoEncoder, Error, TEXT("Failed to destroy NVENC encoder (%s)."), *NVENC.GetErrorString(NVEncoder, result));
 			}
-		}
-
-		if(bAsyncMode)
-		{
-			ProtectedWaitingForPending.Lock();
-			WaitingForPendingActive = false;
-			ProtectedWaitingForPending.Unlock();
-
-			WaitForNextPendingFrame();
-		}
-	}
-
-	void FVideoEncoderNVENC_H264::FLayer::WaitForNextPendingFrame()
-	{
-		ProtectedWaitingForPending.Lock();
-		if (!WaitingForPendingActive)
-		{
-			FInputOutput* NextBuffer = nullptr;
-			if (PendingEncodes.Peek(NextBuffer))
-			{
-				if(bAsyncMode)
-				{
-					Encoder.OnEvent(NextBuffer->CompletionEvent, [this]() { ProcessNextPendingFrame(); });
-					WaitingForPendingActive = true;
-				} 
-				else
-				{
-					ProcessNextPendingFrame();
-				}
-			}
-		}
-		ProtectedWaitingForPending.Unlock();
-	}
-
-	int FVideoEncoderNVENC_H264::FLayer::GetCapability(NV_ENC_CAPS CapsToQuery) const
-	{
-		int					CapsValue = 0;
-		NVENCStruct(NV_ENC_CAPS_PARAM, CapsParam);
-		CapsParam.capsToQuery = CapsToQuery;
-		NVENCSTATUS			Result = NVENC.nvEncGetEncodeCaps(NVEncoder, CodecGUID, &CapsParam, &CapsValue);
-		//	UE_LOG(LogVideoEncoder, Error, TEXT("NVENC.nvEncGetEncodeCaps(NVEncoder, CodecGUID, &CapsParam, &CapsValue); -> %d"), Result);
-		if (Result != NV_ENC_SUCCESS)
-		{
-			UE_LOG(LogVideoEncoder, Warning, TEXT("Failed to query for NVENC capability %d (%s)."), CapsToQuery, *NVENC.GetErrorString(NVEncoder, Result));
-			return 0;
-		}
-		return CapsValue;
-	}
-
-	FVideoEncoderNVENC_H264::FLayer::FInputOutput* FVideoEncoderNVENC_H264::FLayer::GetOrCreateBuffer(const FVideoEncoderInputFrameImpl* InFrame)
-	{
-		void* TextureToCompress = nullptr;
-
-		switch (InFrame->GetFormat())
-		{
-	#if PLATFORM_WINDOWS
-		case EVideoFrameFormat::D3D11_R8G8B8A8_UNORM:
-		case EVideoFrameFormat::D3D12_R8G8B8A8_UNORM:
-			TextureToCompress = InFrame->GetD3D11().EncoderTexture;
-			break;
-	#endif
-	#if WITH_CUDA
-		case EVideoFrameFormat::CUDA_R8G8B8A8_UNORM:
-			TextureToCompress = InFrame->GetCUDA().EncoderTexture;
-			break;
-	#endif
-		case AVEncoder::EVideoFrameFormat::Undefined:
-		default:
-			break;
-		}
-
-		if (!TextureToCompress)
-		{
-			UE_LOG(LogVideoEncoder, Fatal, TEXT("Got passed a null texture to encode."));
-			return nullptr;
-		}
-
-		FInputOutput* Buffer = nullptr;
-		for (FInputOutput* SearchBuffer : CreatedBuffers)
-		{
-			if (SearchBuffer->InputTexture == TextureToCompress)
-			{
-				Buffer = SearchBuffer;
-				break;
-			}
-		}
-		
-		if(Buffer)
-		{
-			//Check for buffer and input texture resolution mismatch
-			if(InFrame->GetWidth() != Buffer->Width || InFrame->GetHeight() != Buffer->Height)
-			{	
-				// Buffer is wrong resolution, destroy it and we will make a new buffer below
-				this->DestroyBuffer(Buffer);
-				Buffer = nullptr;
-			}
-		}
-
-		if (!Buffer)
-		{
-			Buffer = CreateBuffer();
-			Buffer->SourceFrame = static_cast<const FVideoEncoderInputFrameImpl*>(static_cast<const FVideoEncoderInputFrame*>(InFrame)->Obtain());
-
-			if (Buffer && !RegisterInputTexture(*Buffer, TextureToCompress, FIntPoint(InFrame->GetWidth(), InFrame->GetHeight())))
-			{
-				Buffer->SourceFrame->Release();
-				DestroyBuffer(Buffer);
-				Buffer = nullptr;
-			}
-			else
-			{
-				CreatedBuffers.Push(Buffer);
-			}
-		}
-		else
-		{
-			Buffer->SourceFrame = static_cast<const FVideoEncoderInputFrameImpl*>(static_cast<const FVideoEncoderInputFrame*>(InFrame)->Obtain());
-		}
-
-		return Buffer;
-	}
-
-	FVideoEncoderNVENC_H264::FLayer::FInputOutput* FVideoEncoderNVENC_H264::FLayer::CreateBuffer()
+            NVEncoder = nullptr;
+            bAsyncMode = false;
+        }
+    }
+
+    void FVideoEncoderNVENC_H264::FNVENCLayer::ProcessNextPendingFrame()
+    {
+        FInputOutput* Buffer = nullptr;
+        if (PendingEncodes.Dequeue(Buffer))
+        {
+            // this could be a null buffer (for flush)
+            if (Buffer->Width > 0 && Buffer->Height > 0)
+            {
+                // lock output buffers for CPU access
+                if (LockOutputBuffer(*Buffer))
+                {
+                    if (Encoder.OnEncodedPacket)
+                    {
+                        // create packet with buffer contents
+                        FCodecPacketImpl	Packet;
+
+                        Packet.PTS = static_cast<int64>(Buffer->TimeStamp);
+                        Packet.Data = static_cast<const uint8*>(Buffer->BitstreamData);
+                        Packet.DataSize = Buffer->BitstreamDataSize;
+                        if (Buffer->PictureType == NV_ENC_PIC_TYPE_IDR)
+                        {
+                            UE_LOG(LogVideoEncoder, Verbose, TEXT("Generated IDR Frame"));
+                            Packet.IsKeyFrame = true;
+                        }
+                        Packet.VideoQP = Buffer->FrameAvgQP;
+                        Packet.Timings.StartTs = Buffer->EncodeStartTs;
+                        Packet.Timings.FinishTs = FTimespan::FromSeconds(FPlatformTime::Seconds());
+                        Packet.Framerate = EncoderInitParams.frameRateNum;
+
+                        Encoder.OnEncodedPacket(LayerIndex, Buffer->SourceFrame, Packet);
+                    }
+
+                    UnlockOutputBuffer(*Buffer);
+                }
+            }
+
+            // input is no longer required
+            if (Buffer->SourceFrame)
+            {
+                Buffer->SourceFrame->Release();
+                Buffer->SourceFrame = nullptr;
+            }
+
+            // optional event to be triggered
+            if (Buffer->TriggerOnCompletion)
+            {
+                Buffer->TriggerOnCompletion->Trigger();
+            }
+        }
+
+        if(bAsyncMode)
+        {
+            ProtectedWaitingForPending.Lock();
+            WaitingForPendingActive = false;
+            ProtectedWaitingForPending.Unlock();
+
+            WaitForNextPendingFrame();
+        }
+    }
+
+    void FVideoEncoderNVENC_H264::FNVENCLayer::WaitForNextPendingFrame()
+    {
+        ProtectedWaitingForPending.Lock();
+        if (!WaitingForPendingActive)
+        {
+            FInputOutput* NextBuffer = nullptr;
+            if (PendingEncodes.Peek(NextBuffer))
+            {
+                if(bAsyncMode)
+                {
+                    Encoder.OnEvent(NextBuffer->CompletionEvent, [this]() { ProcessNextPendingFrame(); });
+                    WaitingForPendingActive = true;
+                } 
+                else
+                {
+                    ProcessNextPendingFrame();
+                }
+            }
+        }
+        ProtectedWaitingForPending.Unlock();
+    }
+
+    int FVideoEncoderNVENC_H264::FNVENCLayer::GetCapability(NV_ENC_CAPS CapsToQuery) const
+    {
+        int CapsValue = 0;
+        NVENCStruct(NV_ENC_CAPS_PARAM, CapsParam);
+        CapsParam.capsToQuery = CapsToQuery;
+        auto const result = NVENC.nvEncGetEncodeCaps(NVEncoder, CodecGUID, &CapsParam, &CapsValue);
+        if (result != NV_ENC_SUCCESS)
+        {
+            UE_LOG(LogVideoEncoder, Warning, TEXT("Failed to query for NVENC capability %d (%s)."), CapsToQuery, *NVENC.GetErrorString(NVEncoder, result));
+            return 0;
+        }
+        return CapsValue;
+    }
+
+    FVideoEncoderNVENC_H264::FNVENCLayer::FInputOutput* FVideoEncoderNVENC_H264::FNVENCLayer::GetOrCreateBuffer(const FVideoEncoderInputFrameImpl* InFrame)
+    {
+        void* TextureToCompress = nullptr;
+
+        switch (InFrame->GetFormat())
+        {
+    #if PLATFORM_WINDOWS
+        case EVideoFrameFormat::D3D11_R8G8B8A8_UNORM:
+        case EVideoFrameFormat::D3D12_R8G8B8A8_UNORM:
+            TextureToCompress = InFrame->GetD3D11().EncoderTexture;
+            break;
+    #endif
+    #if WITH_CUDA
+        case EVideoFrameFormat::CUDA_R8G8B8A8_UNORM:
+            TextureToCompress = InFrame->GetCUDA().EncoderTexture;
+            break;
+    #endif
+        case AVEncoder::EVideoFrameFormat::Undefined:
+        default:
+            break;
+        }
+
+        if (!TextureToCompress)
+        {
+            UE_LOG(LogVideoEncoder, Fatal, TEXT("Got passed a null texture to encode."));
+            return nullptr;
+        }
+
+        FInputOutput* Buffer = nullptr;
+        for (FInputOutput* SearchBuffer : CreatedBuffers)
+        {
+            if (SearchBuffer->InputTexture == TextureToCompress)
+            {
+                Buffer = SearchBuffer;
+                break;
+            }
+        }
+        
+        if (Buffer)
+        {
+            // Check for buffer and input texture resolution mismatch
+            if (InFrame->GetWidth() != Buffer->Width || InFrame->GetHeight() != Buffer->Height)
+            {	
+                // Buffer is wrong resolution, destroy it and we will make a new buffer below
+                DestroyBuffer(Buffer);
+                Buffer = nullptr;
+            }
+        }
+
+        if (!Buffer)
+        {
+            Buffer = CreateBuffer();
+            Buffer->SourceFrame = static_cast<const FVideoEncoderInputFrameImpl*>(static_cast<const FVideoEncoderInputFrame*>(InFrame)->Obtain());
+
+            if (Buffer && !RegisterInputTexture(*Buffer, TextureToCompress, FIntPoint(InFrame->GetWidth(), InFrame->GetHeight())))
+            {
+                Buffer->SourceFrame->Release();
+                DestroyBuffer(Buffer);
+                Buffer = nullptr;
+            }
+            else
+            {
+                CreatedBuffers.Push(Buffer);
+            }
+        }
+        else
+        {
+            Buffer->SourceFrame = static_cast<const FVideoEncoderInputFrameImpl*>(static_cast<const FVideoEncoderInputFrame*>(InFrame)->Obtain());
+        }
+
+        return Buffer;
+    }
+
+	FVideoEncoderNVENC_H264::FNVENCLayer::FInputOutput* FVideoEncoderNVENC_H264::FNVENCLayer::CreateBuffer()
 	{
 		FInputOutput* Buffer = new FInputOutput();
 
 		// output bit stream buffer
 		NVENCStruct(NV_ENC_CREATE_BITSTREAM_BUFFER, CreateParam);
-		NVENCSTATUS		Result = NVENC.nvEncCreateBitstreamBuffer(NVEncoder, &CreateParam);
-		if (Result != NV_ENC_SUCCESS)
 		{
-			UE_LOG(LogVideoEncoder, Error, TEXT("Failed to create NVENC output buffer (%s)."), *NVENC.GetErrorString(NVEncoder, Result));
-			DestroyBuffer(Buffer);
-			return nullptr;
-		}
-
-		Buffer->OutputBitstream = CreateParam.bitstreamBuffer;
-
-		if(bAsyncMode)
-		{
-	#if PLATFORM_WINDOWS
-			// encode completion async event
-			Buffer->CompletionEvent = CreateEvent(nullptr, false, false, nullptr);
-
-			NVENCStruct(NV_ENC_EVENT_PARAMS, EventParams);
-			EventParams.completionEvent = Buffer->CompletionEvent;
-			Result = NVENC.nvEncRegisterAsyncEvent(NVEncoder, &EventParams);
-			if (Result != NV_ENC_SUCCESS)
+			auto const result = NVENC.nvEncCreateBitstreamBuffer(NVEncoder, &CreateParam);
+			if (result != NV_ENC_SUCCESS)
 			{
-				UE_LOG(LogVideoEncoder, Error, TEXT("Failed to register completion event with NVENC (%s)."), *NVENC.GetErrorString(NVEncoder, Result));
+				UE_LOG(LogVideoEncoder, Error, TEXT("Failed to create NVENC output buffer (%s)."), *NVENC.GetErrorString(NVEncoder, result));
 				DestroyBuffer(Buffer);
 				return nullptr;
 			}
-	#else
-			UE_LOG(LogVideoEncoder, Fatal, TEXT("FVideoEncoderNVENC_H264::FLayer::CreateBuffer should not have hit here as NVENC async mode only works on Windows!"))
-	#endif
-		}
-
-		return Buffer;
-	}
-
-	void FVideoEncoderNVENC_H264::FLayer::DestroyBuffer(FInputOutput* InBuffer)
-	{
-		// unregister input texture - if any
-		UnregisterInputTexture(*InBuffer);
-
-		// destroy output buffer - if any
-		UnlockOutputBuffer(*InBuffer);
-		if (InBuffer->OutputBitstream)
-		{
-			NVENCSTATUS	Result = NVENC.nvEncDestroyBitstreamBuffer(NVEncoder, InBuffer->OutputBitstream);
-			if (Result != NV_ENC_SUCCESS)
-			{
-				UE_LOG(LogVideoEncoder, Warning, TEXT("Failed to destroy NVENC output buffer (%s)."), *NVENC.GetErrorString(NVEncoder, Result));
-			}
-			InBuffer->OutputBitstream = nullptr;
-		}
-
-		// unregister/close completion event - if any
-		if (bAsyncMode && InBuffer->CompletionEvent)
-		{
-	#if PLATFORM_WINDOWS
-			NVENCStruct(NV_ENC_EVENT_PARAMS, EventParams);
-			EventParams.completionEvent = InBuffer->CompletionEvent;
-			NVENCSTATUS	Result = NVENC.nvEncUnregisterAsyncEvent(NVEncoder, &EventParams);
-			if (Result != NV_ENC_SUCCESS)
-			{
-				UE_LOG(LogVideoEncoder, Warning, TEXT("Failed to unregister NVENC completions event (%s)."), *NVENC.GetErrorString(NVEncoder, Result));
-			}
-			CloseHandle(InBuffer->CompletionEvent);
-			InBuffer->CompletionEvent = nullptr;
-	#else
-			UE_LOG(LogVideoEncoder, Fatal, TEXT("FVideoEncoderNVENC_H264::FLayer::DestroyBuffer should not have hit here as NVENC async mode only works on Windows!"))
-	#endif
-		}
-
-		// release source texture
-		InBuffer->InputTexture = nullptr;
-
-		if (InBuffer->SourceFrame)
-		{
-			InBuffer->SourceFrame->Release();
-			InBuffer->SourceFrame = nullptr;
-		}
-
-		// optional event
-		if (InBuffer->TriggerOnCompletion)
-		{
-			FPlatformProcess::ReturnSynchEventToPool(InBuffer->TriggerOnCompletion);
-			InBuffer->TriggerOnCompletion = nullptr;
-		}
-		delete InBuffer;
-	}
-
-	// TODO (M84FIX)
-
-	void FVideoEncoderNVENC_H264::FLayer::CreateResourceDIRECTX(FInputOutput& InBuffer, NV_ENC_REGISTER_RESOURCE& RegisterParam, FIntPoint TextureSize)
-	{
-	#if PLATFORM_WINDOWS
-		D3D11_TEXTURE2D_DESC	Desc;
-		static_cast<ID3D11Texture2D*>(InBuffer.InputTexture)->GetDesc(&Desc);
-		
-		switch (Desc.Format)
-		{
-		case DXGI_FORMAT_NV12:
-			InBuffer.BufferFormat = NV_ENC_BUFFER_FORMAT_NV12;
-			break;
-		case DXGI_FORMAT_R8G8B8A8_UNORM:
-			InBuffer.BufferFormat = NV_ENC_BUFFER_FORMAT_ABGR;
-			break;
-		case DXGI_FORMAT_B8G8R8A8_UNORM:
-			InBuffer.BufferFormat = NV_ENC_BUFFER_FORMAT_ARGB;
-			break;
-		default:
-			UE_LOG(LogVideoEncoder, Error, TEXT("Invalid input texture format for NVENC (%d)"), Desc.Format);
-			return;
-		}
-
-		InBuffer.Width = TextureSize.X;
-		InBuffer.Height = TextureSize.Y;
-		InBuffer.Pitch = 0;
-
-		RegisterParam.resourceType = NV_ENC_INPUT_RESOURCE_TYPE_DIRECTX;
-		RegisterParam.width = Desc.Width;
-		RegisterParam.height = Desc.Height;
-		RegisterParam.pitch = InBuffer.Pitch ;
-		RegisterParam.bufferFormat = InBuffer.BufferFormat;
-		RegisterParam.bufferUsage = NV_ENC_INPUT_IMAGE;
-	#endif
-	}
-
-	void FVideoEncoderNVENC_H264::FLayer::CreateResourceCUDAARRAY(FInputOutput& InBuffer, NV_ENC_REGISTER_RESOURCE& RegisterParam, FIntPoint TextureSize)
-	{
-		InBuffer.Width = TextureSize.X;
-		InBuffer.Height = TextureSize.Y;
-		InBuffer.Pitch = TextureSize.X * 4;
-		InBuffer.BufferFormat = NV_ENC_BUFFER_FORMAT_ARGB;
-
-		RegisterParam.resourceType = NV_ENC_INPUT_RESOURCE_TYPE_CUDAARRAY;
-		RegisterParam.width = InBuffer.Width;
-		RegisterParam.height = InBuffer.Height;
-		RegisterParam.pitch = InBuffer.Pitch;
-		RegisterParam.bufferFormat = InBuffer.BufferFormat;
-		RegisterParam.bufferUsage = NV_ENC_INPUT_IMAGE;
-	}
-
-	bool FVideoEncoderNVENC_H264::FLayer::RegisterInputTexture(FInputOutput& InBuffer, void* InTexture, FIntPoint TextureSize)
-	{
-		if (!InBuffer.InputTexture)
-		{
-			InBuffer.InputTexture = InTexture;
-			NVENCStruct(NV_ENC_REGISTER_RESOURCE, RegisterParam);
-
-			switch (InBuffer.SourceFrame->GetFormat())
-			{
-		#if PLATFORM_WINDOWS
-			case EVideoFrameFormat::D3D11_R8G8B8A8_UNORM:
-			case EVideoFrameFormat::D3D12_R8G8B8A8_UNORM:
-				CreateResourceDIRECTX(InBuffer, RegisterParam, TextureSize);
-				break;
-		#endif
-		#if WITH_CUDA
-			case EVideoFrameFormat::CUDA_R8G8B8A8_UNORM:
-				CreateResourceCUDAARRAY(InBuffer, RegisterParam, TextureSize);
-				break;
-		#endif
-			case AVEncoder::EVideoFrameFormat::Undefined:
-			default:
-				break;
-			}
-
-			RegisterParam.resourceToRegister = InTexture;
-
-			NVENCSTATUS	Result = NVENC.nvEncRegisterResource(NVEncoder, &RegisterParam);
-			if (Result != NV_ENC_SUCCESS)
-			{
-				UE_LOG(LogVideoEncoder, Error, TEXT("Failed to register input texture with NVENC (%s)"), *NVENC.GetErrorString(NVEncoder, Result));
-				return false;
-			}
-			InBuffer.RegisteredInput = RegisterParam.registeredResource;
-		}
-
-		return true;
-	}
-
-	bool FVideoEncoderNVENC_H264::FLayer::UnregisterInputTexture(FInputOutput& InBuffer)
-	{
-		UnmapInputTexture(InBuffer);
-		if (InBuffer.RegisteredInput)
-		{
-			NVENCSTATUS	Result = NVENC.nvEncUnregisterResource(NVEncoder, InBuffer.RegisteredInput);
-			if (Result != NV_ENC_SUCCESS)
-			{
-				UE_LOG(LogVideoEncoder, Error, TEXT("Failed to unregister input texture with NVENC (%s)"), *NVENC.GetErrorString(NVEncoder, Result));
-				InBuffer.InputTexture = nullptr;
-				InBuffer.RegisteredInput = nullptr;
-				return false;
-			}
-			InBuffer.InputTexture = nullptr;
-			InBuffer.RegisteredInput = nullptr;
-		}
-		return true;
-	}
-
-	bool FVideoEncoderNVENC_H264::FLayer::MapInputTexture(FInputOutput& InBuffer)
-	{
-		if (!InBuffer.MappedInput)
-		{
-			NVENCStruct(NV_ENC_MAP_INPUT_RESOURCE, MapInputResource);
-			MapInputResource.registeredResource = InBuffer.RegisteredInput;
-			NVENCSTATUS		Result = NVENC.nvEncMapInputResource(NVEncoder, &MapInputResource);
-			if (Result != NV_ENC_SUCCESS)
-			{
-				UE_LOG(LogVideoEncoder, Error, TEXT("Failed to map input texture buffer (%s)"), *NVENC.GetErrorString(NVEncoder, Result));
-				return false;
-			}
-			InBuffer.MappedInput = MapInputResource.mappedResource;
-			check(InBuffer.BufferFormat == MapInputResource.mappedBufferFmt);
-		}
-		return true;
-	}
-
-	bool FVideoEncoderNVENC_H264::FLayer::UnmapInputTexture(FInputOutput& InBuffer)
-	{
-		if (InBuffer.MappedInput)
-		{
-			NVENCSTATUS		Result = NVENC.nvEncUnmapInputResource(NVEncoder, InBuffer.MappedInput);
-			if (Result != NV_ENC_SUCCESS)
-			{
-				UE_LOG(LogVideoEncoder, Error, TEXT("Failed to unmap input texture buffer (%s)"), *NVENC.GetErrorString(NVEncoder, Result));
-				InBuffer.MappedInput = nullptr;
-				return false;
-			}
-			InBuffer.MappedInput = nullptr;
-		}
-		return true;
-	}
-
-	bool FVideoEncoderNVENC_H264::FLayer::LockOutputBuffer(FInputOutput& InBuffer)
-	{
-		if (!InBuffer.BitstreamData)
-		{
-			// lock output buffers for CPU access
-			NVENCStruct(NV_ENC_LOCK_BITSTREAM, LockBitstreamParam);
-			LockBitstreamParam.outputBitstream = InBuffer.OutputBitstream;
-			NVENCSTATUS Result = NVENC.nvEncLockBitstream(NVEncoder, &LockBitstreamParam);
-			if (Result != NV_ENC_SUCCESS)
-			{
-				UE_LOG(LogVideoEncoder, Error, TEXT("Failed to lock output bitstream for NVENC encoder (%s)."), *NVENC.GetErrorString(NVEncoder, Result));
-				return false;
-			}
-			else
-			{
-				InBuffer.BitstreamData = LockBitstreamParam.bitstreamBufferPtr;
-				InBuffer.BitstreamDataSize = LockBitstreamParam.bitstreamSizeInBytes;
-				InBuffer.PictureType = LockBitstreamParam.pictureType;
-				InBuffer.FrameAvgQP = LockBitstreamParam.frameAvgQP;
-				InBuffer.TimeStamp = LockBitstreamParam.outputTimeStamp;
-			}
 		}
-		return true;
-	}
 
-	bool FVideoEncoderNVENC_H264::FLayer::UnlockOutputBuffer(FInputOutput& InBuffer)
-	{
-		if (InBuffer.BitstreamData)
-		{
-			NVENCSTATUS		Result = NVENC.nvEncUnlockBitstream(NVEncoder, InBuffer.OutputBitstream);
-			if (Result != NV_ENC_SUCCESS)
-			{
-				UE_LOG(LogVideoEncoder, Error, TEXT("Failed to unlock output bitstream for NVENC encoder (%s)."), *NVENC.GetErrorString(NVEncoder, Result));
-				return false;
-			}
-			else
-			{
-				InBuffer.BitstreamData = nullptr;
-				InBuffer.BitstreamDataSize = 0;
-			}
-		}
-		return true;
-	}
+        Buffer->OutputBitstream = CreateParam.bitstreamBuffer;
+
+        if (bAsyncMode)
+        {
+    #if PLATFORM_WINDOWS
+            // encode completion async event
+            Buffer->CompletionEvent = CreateEvent(nullptr, false, false, nullptr);
+
+            NVENCStruct(NV_ENC_EVENT_PARAMS, EventParams);
+            EventParams.completionEvent = Buffer->CompletionEvent;
+            auto const result = NVENC.nvEncRegisterAsyncEvent(NVEncoder, &EventParams);
+            if (result != NV_ENC_SUCCESS)
+            {
+                UE_LOG(LogVideoEncoder, Error, TEXT("Failed to register completion event with NVENC (%s)."), *NVENC.GetErrorString(NVEncoder, result));
+                DestroyBuffer(Buffer);
+                return nullptr;
+            }
+    #else
+            UE_LOG(LogVideoEncoder, Fatal, TEXT("FVideoEncoderNVENC_H264::FNVENCLayer::CreateBuffer should not have hit here as NVENC async mode only works on Windows!"))
+    #endif
+        }
+
+        return Buffer;
+    }
+
+    void FVideoEncoderNVENC_H264::FNVENCLayer::DestroyBuffer(FInputOutput* InBuffer)
+    {
+        // unregister input texture - if any
+        UnregisterInputTexture(*InBuffer);
+
+        // destroy output buffer - if any
+        UnlockOutputBuffer(*InBuffer);
+        if (InBuffer->OutputBitstream)
+        {
+            auto const result = NVENC.nvEncDestroyBitstreamBuffer(NVEncoder, InBuffer->OutputBitstream);
+            if (result != NV_ENC_SUCCESS)
+            {
+                UE_LOG(LogVideoEncoder, Warning, TEXT("Failed to destroy NVENC output buffer (%s)."), *NVENC.GetErrorString(NVEncoder, result));
+            }
+            InBuffer->OutputBitstream = nullptr;
+        }
+
+        // unregister/close completion event - if any
+        if (bAsyncMode && InBuffer->CompletionEvent)
+        {
+    #if PLATFORM_WINDOWS
+            NVENCStruct(NV_ENC_EVENT_PARAMS, EventParams);
+            EventParams.completionEvent = InBuffer->CompletionEvent;
+            auto const result = NVENC.nvEncUnregisterAsyncEvent(NVEncoder, &EventParams);
+            if (result != NV_ENC_SUCCESS)
+            {
+                UE_LOG(LogVideoEncoder, Warning, TEXT("Failed to unregister NVENC completions event (%s)."), *NVENC.GetErrorString(NVEncoder, result));
+            }
+            CloseHandle(InBuffer->CompletionEvent);
+            InBuffer->CompletionEvent = nullptr;
+    #else
+            UE_LOG(LogVideoEncoder, Fatal, TEXT("FVideoEncoderNVENC_H264::FNVENCLayer::DestroyBuffer should not have hit here as NVENC async mode only works on Windows!"))
+    #endif
+        }
+
+        // release source texture
+        InBuffer->InputTexture = nullptr;
+
+        if (InBuffer->SourceFrame)
+        {
+            InBuffer->SourceFrame->Release();
+            InBuffer->SourceFrame = nullptr;
+        }
+
+        // optional event
+        if (InBuffer->TriggerOnCompletion)
+        {
+            FPlatformProcess::ReturnSynchEventToPool(InBuffer->TriggerOnCompletion);
+            InBuffer->TriggerOnCompletion = nullptr;
+        }
+        delete InBuffer;
+    }
+
+    // TODO (M84FIX)
+
+    void FVideoEncoderNVENC_H264::FNVENCLayer::CreateResourceDIRECTX(FInputOutput& InBuffer, NV_ENC_REGISTER_RESOURCE& RegisterParam, FIntPoint TextureSize)
+    {
+    #if PLATFORM_WINDOWS
+        D3D11_TEXTURE2D_DESC	Desc;
+        static_cast<ID3D11Texture2D*>(InBuffer.InputTexture)->GetDesc(&Desc);
+        
+        switch (Desc.Format)
+        {
+        case DXGI_FORMAT_NV12:
+            InBuffer.BufferFormat = NV_ENC_BUFFER_FORMAT_NV12;
+            break;
+        case DXGI_FORMAT_R8G8B8A8_UNORM:
+            InBuffer.BufferFormat = NV_ENC_BUFFER_FORMAT_ABGR;
+            break;
+        case DXGI_FORMAT_B8G8R8A8_UNORM:
+            InBuffer.BufferFormat = NV_ENC_BUFFER_FORMAT_ARGB;
+            break;
+        default:
+            UE_LOG(LogVideoEncoder, Error, TEXT("Invalid input texture format for NVENC (%d)"), Desc.Format);
+            return;
+        }
+
+        InBuffer.Width = TextureSize.X;
+        InBuffer.Height = TextureSize.Y;
+        InBuffer.Pitch = 0;
+
+        RegisterParam.resourceType = NV_ENC_INPUT_RESOURCE_TYPE_DIRECTX;
+        RegisterParam.width = Desc.Width;
+        RegisterParam.height = Desc.Height;
+        RegisterParam.pitch = InBuffer.Pitch ;
+        RegisterParam.bufferFormat = InBuffer.BufferFormat;
+        RegisterParam.bufferUsage = NV_ENC_INPUT_IMAGE;
+    #endif
+    }
+
+    void FVideoEncoderNVENC_H264::FNVENCLayer::CreateResourceCUDAARRAY(FInputOutput& InBuffer, NV_ENC_REGISTER_RESOURCE& RegisterParam, FIntPoint TextureSize)
+    {
+        InBuffer.Width = TextureSize.X;
+        InBuffer.Height = TextureSize.Y;
+        InBuffer.Pitch = TextureSize.X * 4;
+        InBuffer.BufferFormat = NV_ENC_BUFFER_FORMAT_ARGB;
+
+        RegisterParam.resourceType = NV_ENC_INPUT_RESOURCE_TYPE_CUDAARRAY;
+        RegisterParam.width = InBuffer.Width;
+        RegisterParam.height = InBuffer.Height;
+        RegisterParam.pitch = InBuffer.Pitch;
+        RegisterParam.bufferFormat = InBuffer.BufferFormat;
+        RegisterParam.bufferUsage = NV_ENC_INPUT_IMAGE;
+    }
+
+    bool FVideoEncoderNVENC_H264::FNVENCLayer::RegisterInputTexture(FInputOutput& InBuffer, void* InTexture, FIntPoint TextureSize)
+    {
+        if (!InBuffer.InputTexture)
+        {
+            InBuffer.InputTexture = InTexture;
+            NVENCStruct(NV_ENC_REGISTER_RESOURCE, RegisterParam);
+
+            switch (InBuffer.SourceFrame->GetFormat())
+            {
+        #if PLATFORM_WINDOWS
+            case EVideoFrameFormat::D3D11_R8G8B8A8_UNORM:
+            case EVideoFrameFormat::D3D12_R8G8B8A8_UNORM:
+                CreateResourceDIRECTX(InBuffer, RegisterParam, TextureSize);
+                break;
+        #endif
+        #if WITH_CUDA
+            case EVideoFrameFormat::CUDA_R8G8B8A8_UNORM:
+                CreateResourceCUDAARRAY(InBuffer, RegisterParam, TextureSize);
+                break;
+        #endif
+            case AVEncoder::EVideoFrameFormat::Undefined:
+            default:
+                break;
+            }
+
+            RegisterParam.resourceToRegister = InTexture;
+
+            NVENCSTATUS	Result = NVENC.nvEncRegisterResource(NVEncoder, &RegisterParam);
+            if (Result != NV_ENC_SUCCESS)
+            {
+                UE_LOG(LogVideoEncoder, Error, TEXT("Failed to register input texture with NVENC (%s)"), *NVENC.GetErrorString(NVEncoder, Result));
+                return false;
+            }
+            InBuffer.RegisteredInput = RegisterParam.registeredResource;
+        }
+
+        return true;
+    }
+
+    bool FVideoEncoderNVENC_H264::FNVENCLayer::UnregisterInputTexture(FInputOutput& InBuffer)
+    {
+        UnmapInputTexture(InBuffer);
+        if (InBuffer.RegisteredInput)
+        {
+            NVENCSTATUS	Result = NVENC.nvEncUnregisterResource(NVEncoder, InBuffer.RegisteredInput);
+            if (Result != NV_ENC_SUCCESS)
+            {
+                UE_LOG(LogVideoEncoder, Error, TEXT("Failed to unregister input texture with NVENC (%s)"), *NVENC.GetErrorString(NVEncoder, Result));
+                InBuffer.InputTexture = nullptr;
+                InBuffer.RegisteredInput = nullptr;
+                return false;
+            }
+            InBuffer.InputTexture = nullptr;
+            InBuffer.RegisteredInput = nullptr;
+        }
+        return true;
+    }
+
+    bool FVideoEncoderNVENC_H264::FNVENCLayer::MapInputTexture(FInputOutput& InBuffer)
+    {
+        if (!InBuffer.MappedInput)
+        {
+            NVENCStruct(NV_ENC_MAP_INPUT_RESOURCE, MapInputResource);
+            MapInputResource.registeredResource = InBuffer.RegisteredInput;
+            NVENCSTATUS		Result = NVENC.nvEncMapInputResource(NVEncoder, &MapInputResource);
+            if (Result != NV_ENC_SUCCESS)
+            {
+                UE_LOG(LogVideoEncoder, Error, TEXT("Failed to map input texture buffer (%s)"), *NVENC.GetErrorString(NVEncoder, Result));
+                return false;
+            }
+            InBuffer.MappedInput = MapInputResource.mappedResource;
+            check(InBuffer.BufferFormat == MapInputResource.mappedBufferFmt);
+        }
+        return true;
+    }
+
+    bool FVideoEncoderNVENC_H264::FNVENCLayer::UnmapInputTexture(FInputOutput& InBuffer)
+    {
+        if (InBuffer.MappedInput)
+        {
+            NVENCSTATUS		Result = NVENC.nvEncUnmapInputResource(NVEncoder, InBuffer.MappedInput);
+            if (Result != NV_ENC_SUCCESS)
+            {
+                UE_LOG(LogVideoEncoder, Error, TEXT("Failed to unmap input texture buffer (%s)"), *NVENC.GetErrorString(NVEncoder, Result));
+                InBuffer.MappedInput = nullptr;
+                return false;
+            }
+            InBuffer.MappedInput = nullptr;
+        }
+        return true;
+    }
+
+    bool FVideoEncoderNVENC_H264::FNVENCLayer::LockOutputBuffer(FInputOutput& InBuffer)
+    {
+        if (!InBuffer.BitstreamData)
+        {
+            // lock output buffers for CPU access
+            NVENCStruct(NV_ENC_LOCK_BITSTREAM, LockBitstreamParam);
+            LockBitstreamParam.outputBitstream = InBuffer.OutputBitstream;
+            NVENCSTATUS Result = NVENC.nvEncLockBitstream(NVEncoder, &LockBitstreamParam);
+            if (Result != NV_ENC_SUCCESS)
+            {
+                UE_LOG(LogVideoEncoder, Error, TEXT("Failed to lock output bitstream for NVENC encoder (%s)."), *NVENC.GetErrorString(NVEncoder, Result));
+                return false;
+            }
+            else
+            {
+                InBuffer.BitstreamData = LockBitstreamParam.bitstreamBufferPtr;
+                InBuffer.BitstreamDataSize = LockBitstreamParam.bitstreamSizeInBytes;
+                InBuffer.PictureType = LockBitstreamParam.pictureType;
+                InBuffer.FrameAvgQP = LockBitstreamParam.frameAvgQP;
+                InBuffer.TimeStamp = LockBitstreamParam.outputTimeStamp;
+            }
+        }
+        return true;
+    }
+
+    bool FVideoEncoderNVENC_H264::FNVENCLayer::UnlockOutputBuffer(FInputOutput& InBuffer)
+    {
+        if (InBuffer.BitstreamData)
+        {
+            NVENCSTATUS		Result = NVENC.nvEncUnlockBitstream(NVEncoder, InBuffer.OutputBitstream);
+            if (Result != NV_ENC_SUCCESS)
+            {
+                UE_LOG(LogVideoEncoder, Error, TEXT("Failed to unlock output bitstream for NVENC encoder (%s)."), *NVENC.GetErrorString(NVEncoder, Result));
+                return false;
+            }
+            else
+            {
+                InBuffer.BitstreamData = nullptr;
+                InBuffer.BitstreamDataSize = 0;
+            }
+        }
+        return true;
+    }
 
 #if PLATFORM_WINDOWS
-	static bool CreateEncoderDevice(TRefCountPtr<ID3D11Device>& OutEncoderDevice, TRefCountPtr<ID3D11DeviceContext>& OutEncoderDeviceContext)
-	{
-		// need a d3d11 context to be able to set up encoder
-		TRefCountPtr<IDXGIFactory1>			DXGIFactory1;
-		TRefCountPtr<IDXGIAdapter>			Adapter;
-
-		HRESULT		Result = CreateDXGIFactory1(__uuidof(IDXGIFactory1), (void**)DXGIFactory1.GetInitReference());
-		if (Result != S_OK)
-		{
-			UE_LOG(LogVideoEncoder, Error, TEXT("Failed to create DX factory for NVENC."));
-			return false;
-		}
-
-		for (int GpuIndex = 0; GpuIndex < MAX_GPU_INDEXES; GpuIndex++)
-		{
-			if ((Result = DXGIFactory1->EnumAdapters(GpuIndex, Adapter.GetInitReference())) != S_OK)
-			{
-				UE_LOG(LogVideoEncoder, Error, TEXT("Failed to enum GPU #%d for NVENC."), GpuIndex);
-				return false;
-			}
-
-			DXGI_ADAPTER_DESC AdapterDesc;
-			Adapter->GetDesc(&AdapterDesc);
-			if (AdapterDesc.VendorId != 0x10DE) // NVIDIA
-			{
-				continue;
-			}
-
-			if ((Result = D3D11CreateDevice(Adapter, D3D_DRIVER_TYPE_UNKNOWN, NULL, 0,
-				NULL, 0, D3D11_SDK_VERSION, OutEncoderDevice.GetInitReference(),
-				NULL, OutEncoderDeviceContext.GetInitReference())) != S_OK)
-			{
-				UE_LOG(LogVideoEncoder, Error, TEXT("Failed to create D3D11 device for NVENC."));
-			}
-			else
-			{
-				UE_LOG(LogVideoEncoder, Log, TEXT("Created D3D11 device for NVENC on '%s'."), AdapterDesc.Description);
-				return true;
-			}
-		}
-
-		UE_LOG(LogVideoEncoder, Error, TEXT("No compatible devices found for NVENC."));
-		return false;
-	}
+    static bool CreateEncoderDevice(TRefCountPtr<ID3D11Device>& OutEncoderDevice, TRefCountPtr<ID3D11DeviceContext>& OutEncoderDeviceContext)
+    {
+        // need a d3d11 context to be able to set up encoder
+        TRefCountPtr<IDXGIFactory1>			DXGIFactory1;
+        TRefCountPtr<IDXGIAdapter>			Adapter;
+
+        HRESULT		Result = CreateDXGIFactory1(__uuidof(IDXGIFactory1), (void**)DXGIFactory1.GetInitReference());
+        if (Result != S_OK)
+        {
+            UE_LOG(LogVideoEncoder, Error, TEXT("Failed to create DX factory for NVENC."));
+            return false;
+        }
+
+        for (int GpuIndex = 0; GpuIndex < MAX_GPU_INDEXES; GpuIndex++)
+        {
+            if ((Result = DXGIFactory1->EnumAdapters(GpuIndex, Adapter.GetInitReference())) != S_OK)
+            {
+                UE_LOG(LogVideoEncoder, Error, TEXT("Failed to enum GPU #%d for NVENC."), GpuIndex);
+                return false;
+            }
+
+            DXGI_ADAPTER_DESC AdapterDesc;
+            Adapter->GetDesc(&AdapterDesc);
+            if (AdapterDesc.VendorId != 0x10DE) // NVIDIA
+            {
+                continue;
+            }
+
+            if ((Result = D3D11CreateDevice(Adapter, D3D_DRIVER_TYPE_UNKNOWN, NULL, 0,
+                NULL, 0, D3D11_SDK_VERSION, OutEncoderDevice.GetInitReference(),
+                NULL, OutEncoderDeviceContext.GetInitReference())) != S_OK)
+            {
+                UE_LOG(LogVideoEncoder, Error, TEXT("Failed to create D3D11 device for NVENC."));
+            }
+            else
+            {
+                UE_LOG(LogVideoEncoder, Log, TEXT("Created D3D11 device for NVENC on '%s'."), AdapterDesc.Description);
+                return true;
+            }
+        }
+
+        UE_LOG(LogVideoEncoder, Error, TEXT("No compatible devices found for NVENC."));
+        return false;
+    }
 #endif
 
 #if PLATFORM_WINDOWS
-	static void* CreateEncoderSession(FNVENCCommon& NVENC, TRefCountPtr<ID3D11Device> InD3D11Device)
-	{
-		void* EncoderSession = nullptr;
-		// create the encoder session
-		NVENCStruct(NV_ENC_OPEN_ENCODE_SESSION_EX_PARAMS, OpenEncodeSessionExParams);
-		OpenEncodeSessionExParams.device = InD3D11Device;
-		OpenEncodeSessionExParams.deviceType = NV_ENC_DEVICE_TYPE_DIRECTX;	// Currently only DX11 is supported
-		OpenEncodeSessionExParams.apiVersion = NVENCAPI_VERSION;
-
-		NVENCSTATUS		NvResult = NVENC.nvEncOpenEncodeSessionEx(&OpenEncodeSessionExParams, &EncoderSession);
-		// UE_LOG(LogVideoEncoder, Error, TEXT("NVENC.nvEncOpenEncodeSessionEx(&OpenEncodeSessionExParams, &EncoderSession); -> %d"), NvResult);
-		if (NvResult != NV_ENC_SUCCESS)
-		{
-			UE_LOG(LogVideoEncoder, Error, TEXT("Unable to open NvEnc encoding session (status: %d)."), NvResult);
-			EncoderSession = nullptr;
-		}
-		return EncoderSession;
-	}
+    static void* CreateEncoderSession(FNVENCCommon& NVENC, TRefCountPtr<ID3D11Device> InD3D11Device)
+    {
+        void* EncoderSession = nullptr;
+        // create the encoder session
+        NVENCStruct(NV_ENC_OPEN_ENCODE_SESSION_EX_PARAMS, OpenEncodeSessionExParams);
+        OpenEncodeSessionExParams.device = InD3D11Device;
+        OpenEncodeSessionExParams.deviceType = NV_ENC_DEVICE_TYPE_DIRECTX;	// Currently only DX11 is supported
+        OpenEncodeSessionExParams.apiVersion = NVENCAPI_VERSION;
+
+        NVENCSTATUS		NvResult = NVENC.nvEncOpenEncodeSessionEx(&OpenEncodeSessionExParams, &EncoderSession);
+        // UE_LOG(LogVideoEncoder, Error, TEXT("NVENC.nvEncOpenEncodeSessionEx(&OpenEncodeSessionExParams, &EncoderSession); -> %d"), NvResult);
+        if (NvResult != NV_ENC_SUCCESS)
+        {
+            UE_LOG(LogVideoEncoder, Error, TEXT("Unable to open NvEnc encoding session (status: %d)."), NvResult);
+            EncoderSession = nullptr;
+        }
+        return EncoderSession;
+    }
 #endif // PLATFORM_WINDOWS
-	
+    
 #if WITH_CUDA
-	static void* CreateEncoderSession(FNVENCCommon& NVENC, CUcontext CudaContext)
-	{
-		void* EncoderSession = nullptr;
-		// create the encoder session
-		NVENCStruct(NV_ENC_OPEN_ENCODE_SESSION_EX_PARAMS, OpenEncodeSessionExParams);
-		OpenEncodeSessionExParams.device = CudaContext;
-		OpenEncodeSessionExParams.deviceType = NV_ENC_DEVICE_TYPE_CUDA;	// We use cuda to pass vulkan device memory to nvenc
-		OpenEncodeSessionExParams.apiVersion = NVENCAPI_VERSION;
-
-		NVENCSTATUS		NvResult = NVENC.nvEncOpenEncodeSessionEx(&OpenEncodeSessionExParams, &EncoderSession);
-		//	UE_LOG(LogVideoEncoder, Error, TEXT("NVENC.nvEncOpenEncodeSessionEx(&OpenEncodeSessionExParams, &EncoderSession); -> %d"), NvResult);
-		if (NvResult != NV_ENC_SUCCESS)
-		{
-			UE_LOG(LogVideoEncoder, Error, TEXT("Unable to open NvEnc encoding session (status: %d)."), NvResult);
-			EncoderSession = nullptr;
-		}
-		return EncoderSession;
-	}
+    static void* CreateEncoderSession(FNVENCCommon& NVENC, CUcontext CudaContext)
+    {
+        void* EncoderSession = nullptr;
+        // create the encoder session
+        NVENCStruct(NV_ENC_OPEN_ENCODE_SESSION_EX_PARAMS, OpenEncodeSessionExParams);
+        OpenEncodeSessionExParams.device = CudaContext;
+        OpenEncodeSessionExParams.deviceType = NV_ENC_DEVICE_TYPE_CUDA;	// We use cuda to pass vulkan device memory to nvenc
+        OpenEncodeSessionExParams.apiVersion = NVENCAPI_VERSION;
+
+        NVENCSTATUS		NvResult = NVENC.nvEncOpenEncodeSessionEx(&OpenEncodeSessionExParams, &EncoderSession);
+        //	UE_LOG(LogVideoEncoder, Error, TEXT("NVENC.nvEncOpenEncodeSessionEx(&OpenEncodeSessionExParams, &EncoderSession); -> %d"), NvResult);
+        if (NvResult != NV_ENC_SUCCESS)
+        {
+            UE_LOG(LogVideoEncoder, Error, TEXT("Unable to open NvEnc encoding session (status: %d)."), NvResult);
+            EncoderSession = nullptr;
+        }
+        return EncoderSession;
+    }
 #endif
 
-	static int GetEncoderCapability(FNVENCCommon& NVENC, void* InEncoder, NV_ENC_CAPS InCapsToQuery)
-	{
-		int					CapsValue = 0;
-		NVENCStruct(NV_ENC_CAPS_PARAM, CapsParam);
-		CapsParam.capsToQuery = InCapsToQuery;
-		NVENCSTATUS Result = NVENC.nvEncGetEncodeCaps(InEncoder, NV_ENC_CODEC_H264_GUID, &CapsParam, &CapsValue);
-		//	UE_LOG(LogVideoEncoder, Error, TEXT("NVENC.nvEncGetEncodeCaps(InEncoder, NV_ENC_CODEC_H264_GUID, &CapsParam, &CapsValue); -> %d"), Result);
-		if (Result != NV_ENC_SUCCESS)
-		{
-			UE_LOG(LogVideoEncoder, Warning, TEXT("Failed to query for NVENC capability %d (error %d)."), InCapsToQuery, Result);
-			return 0;
-		}
-		return CapsValue;
-	}
-
-	static bool GetEncoderSupportedProfiles(FNVENCCommon& NVENC, void* InEncoder, uint32& OutSupportedProfiles)
-	{
-		const uint32		MaxProfileGUIDs = 32;
-		GUID				ProfileGUIDs[MaxProfileGUIDs];
-		uint32				NumProfileGUIDs = 0;
-
-		OutSupportedProfiles = 0;
-		NVENCSTATUS		Result = NVENC.nvEncGetEncodeProfileGUIDs(InEncoder, NV_ENC_CODEC_H264_GUID, ProfileGUIDs, MaxProfileGUIDs, &NumProfileGUIDs);
-		//	UE_LOG(LogVideoEncoder, Error, TEXT("NVENC.nvEncGetEncodeProfileGUIDs(InEncoder, NV_ENC_CODEC_H264_GUID, ProfileGUIDs, MaxProfileGUIDs, &NumProfileGUIDs); -> %d"), Result);
-		if (Result != NV_ENC_SUCCESS)
-		{
-			UE_LOG(LogVideoEncoder, Error, TEXT("Unable to query profiles supported by NvEnc (error: %d)."), Result);
-			return false;
-		}
-		for (uint32 Index = 0; Index < NumProfileGUIDs; ++Index)
-		{
-			if (memcmp(&NV_ENC_H264_PROFILE_BASELINE_GUID, &ProfileGUIDs[Index], sizeof(GUID)) == 0)
-			{
-				OutSupportedProfiles |= H264Profile_Baseline;
-				if (GetEncoderCapability(NVENC, InEncoder, NV_ENC_CAPS_SUPPORT_CONSTRAINED_ENCODING))
-				{
-					OutSupportedProfiles |= H264Profile_ConstrainedBaseline;
-				}
-			}
-			else if (memcmp(&NV_ENC_H264_PROFILE_MAIN_GUID, &ProfileGUIDs[Index], sizeof(GUID)) == 0)
-			{
-				OutSupportedProfiles |= H264Profile_Main;
-			}
-			else if (memcmp(&NV_ENC_H264_PROFILE_HIGH_GUID, &ProfileGUIDs[Index], sizeof(GUID)) == 0)
-			{
-				OutSupportedProfiles |= H264Profile_High;
-			}
-			else if (memcmp(&NV_ENC_H264_PROFILE_CONSTRAINED_HIGH_GUID, &ProfileGUIDs[Index], sizeof(GUID)) == 0)
-			{
-				OutSupportedProfiles |= H264Profile_ConstrainedHigh;
-			}
-		}
-		return OutSupportedProfiles != 0;
-	}
-
-	static bool GetEncoderSupportedInputFormats(FNVENCCommon& NVENC, void* InEncoder, TArray<EVideoFrameFormat>& OutSupportedInputFormats)
-	{
-		const uint32_t			MaxInputFmtCount = 32;
-		uint32_t				InputFmtCount = 0;
-		NV_ENC_BUFFER_FORMAT	InputFormats[MaxInputFmtCount];
-		NVENCSTATUS		Result = NVENC.nvEncGetInputFormats(InEncoder, NV_ENC_CODEC_H264_GUID, InputFormats, MaxInputFmtCount, &InputFmtCount);
-		if (Result != NV_ENC_SUCCESS)
-		{
-			UE_LOG(LogVideoEncoder, Error, TEXT("Unable to query input formats supported by NvEnc (error: %d)."), Result);
-			return false;
-		}
-
-		for (uint32_t Index = 0; Index < InputFmtCount; ++Index)
-		{
-			switch (InputFormats[Index])
-			{
-			case NV_ENC_BUFFER_FORMAT_IYUV:
-				break;
-			case NV_ENC_BUFFER_FORMAT_NV12:
-				break;
-			case NV_ENC_BUFFER_FORMAT_ARGB:
-				break;
-			case NV_ENC_BUFFER_FORMAT_ABGR:
-			#if PLATFORM_WINDOWS
-				OutSupportedInputFormats.Push(EVideoFrameFormat::D3D11_R8G8B8A8_UNORM);
-				OutSupportedInputFormats.Push(EVideoFrameFormat::D3D12_R8G8B8A8_UNORM);
-			#endif
-			#if WITH_CUDA
-				OutSupportedInputFormats.Push(EVideoFrameFormat::CUDA_R8G8B8A8_UNORM);
-			#endif
-				break;
-			}
-		}
-		return true;
-	}
-
-	static bool GetEncoderInfo(FNVENCCommon& NVENC, FVideoEncoderInfo& EncoderInfo)
-	{
-		bool		bSuccess = true;
-		// create a temporary encoder session
-
-		void* EncoderSession = nullptr;
-
-	#if PLATFORM_WINDOWS
-		TRefCountPtr<ID3D11Device>			EncoderDevice;
-		TRefCountPtr<ID3D11DeviceContext>	EncoderDeviceContext;
-
-		if (!CreateEncoderDevice(EncoderDevice, EncoderDeviceContext))
-		{
-			bSuccess = false;
-		}
-		if ((EncoderSession = CreateEncoderSession(NVENC, EncoderDevice)) == nullptr)
-		{
-			bSuccess = false;
-		}
-	#endif
-
-	#if WITH_CUDA
-		if (!EncoderSession)
-		{
-			EncoderSession = CreateEncoderSession(NVENC, FModuleManager::GetModuleChecked<FCUDAModule>("CUDA").GetCudaContext());
-		}
-	#endif
-
-		EncoderInfo.CodecType = ECodecType::H264;
-		EncoderInfo.MaxWidth = GetEncoderCapability(NVENC, EncoderSession, NV_ENC_CAPS_WIDTH_MAX);
-		EncoderInfo.MaxHeight = GetEncoderCapability(NVENC, EncoderSession, NV_ENC_CAPS_HEIGHT_MAX);
-
-		int LevelMax = GetEncoderCapability(NVENC, EncoderSession, NV_ENC_CAPS_LEVEL_MAX);
-		int LevelMin = GetEncoderCapability(NVENC, EncoderSession, NV_ENC_CAPS_LEVEL_MIN);
-		if (LevelMin > 0 && LevelMax > 0 && LevelMax >= LevelMin)
-		{
-			EncoderInfo.H264.MinLevel = (LevelMin > 9) ? LevelMin : 9;
-			EncoderInfo.H264.MaxLevel = (LevelMax < 9) ? 9 : (LevelMax > NV_ENC_LEVEL_H264_52) ? NV_ENC_LEVEL_H264_52 : LevelMax;
-		}
-		else
-		{
-			UE_LOG(LogVideoEncoder, Error, TEXT("Failed to query min/max h264 level supported by NvEnc (reported min/max=%d/%d)."), LevelMin, LevelMax);
-			bSuccess = false;
-		}
-
-		if (!GetEncoderSupportedProfiles(NVENC, EncoderSession, EncoderInfo.H264.SupportedProfiles) ||
-			!GetEncoderSupportedInputFormats(NVENC, EncoderSession, EncoderInfo.SupportedInputFormats))
-		{
-			bSuccess = false;
-		}
-
-		// destroy encoder session
-		if (EncoderSession)
-		{
-			NVENC.nvEncDestroyEncoder(EncoderSession);
-		}
-
-		return bSuccess;
-	}
+    static int GetEncoderCapability(FNVENCCommon& NVENC, void* InEncoder, NV_ENC_CAPS InCapsToQuery)
+    {
+        int					CapsValue = 0;
+        NVENCStruct(NV_ENC_CAPS_PARAM, CapsParam);
+        CapsParam.capsToQuery = InCapsToQuery;
+        NVENCSTATUS Result = NVENC.nvEncGetEncodeCaps(InEncoder, NV_ENC_CODEC_H264_GUID, &CapsParam, &CapsValue);
+        //	UE_LOG(LogVideoEncoder, Error, TEXT("NVENC.nvEncGetEncodeCaps(InEncoder, NV_ENC_CODEC_H264_GUID, &CapsParam, &CapsValue); -> %d"), Result);
+        if (Result != NV_ENC_SUCCESS)
+        {
+            UE_LOG(LogVideoEncoder, Warning, TEXT("Failed to query for NVENC capability %d (error %d)."), InCapsToQuery, Result);
+            return 0;
+        }
+        return CapsValue;
+    }
+
+    static bool GetEncoderSupportedProfiles(FNVENCCommon& NVENC, void* InEncoder, uint32& OutSupportedProfiles)
+    {
+        const uint32		MaxProfileGUIDs = 32;
+        GUID				ProfileGUIDs[MaxProfileGUIDs];
+        uint32				NumProfileGUIDs = 0;
+
+        OutSupportedProfiles = 0;
+        NVENCSTATUS		Result = NVENC.nvEncGetEncodeProfileGUIDs(InEncoder, NV_ENC_CODEC_H264_GUID, ProfileGUIDs, MaxProfileGUIDs, &NumProfileGUIDs);
+        //	UE_LOG(LogVideoEncoder, Error, TEXT("NVENC.nvEncGetEncodeProfileGUIDs(InEncoder, NV_ENC_CODEC_H264_GUID, ProfileGUIDs, MaxProfileGUIDs, &NumProfileGUIDs); -> %d"), Result);
+        if (Result != NV_ENC_SUCCESS)
+        {
+            UE_LOG(LogVideoEncoder, Error, TEXT("Unable to query profiles supported by NvEnc (error: %d)."), Result);
+            return false;
+        }
+        for (uint32 Index = 0; Index < NumProfileGUIDs; ++Index)
+        {
+            if (memcmp(&NV_ENC_H264_PROFILE_BASELINE_GUID, &ProfileGUIDs[Index], sizeof(GUID)) == 0)
+            {
+                OutSupportedProfiles |= H264Profile_Baseline;
+                if (GetEncoderCapability(NVENC, InEncoder, NV_ENC_CAPS_SUPPORT_CONSTRAINED_ENCODING))
+                {
+                    OutSupportedProfiles |= H264Profile_ConstrainedBaseline;
+                }
+            }
+            else if (memcmp(&NV_ENC_H264_PROFILE_MAIN_GUID, &ProfileGUIDs[Index], sizeof(GUID)) == 0)
+            {
+                OutSupportedProfiles |= H264Profile_Main;
+            }
+            else if (memcmp(&NV_ENC_H264_PROFILE_HIGH_GUID, &ProfileGUIDs[Index], sizeof(GUID)) == 0)
+            {
+                OutSupportedProfiles |= H264Profile_High;
+            }
+            else if (memcmp(&NV_ENC_H264_PROFILE_CONSTRAINED_HIGH_GUID, &ProfileGUIDs[Index], sizeof(GUID)) == 0)
+            {
+                OutSupportedProfiles |= H264Profile_ConstrainedHigh;
+            }
+        }
+        return OutSupportedProfiles != 0;
+    }
+
+    static bool GetEncoderSupportedInputFormats(FNVENCCommon& NVENC, void* InEncoder, TArray<EVideoFrameFormat>& OutSupportedInputFormats)
+    {
+        const uint32_t			MaxInputFmtCount = 32;
+        uint32_t				InputFmtCount = 0;
+        NV_ENC_BUFFER_FORMAT	InputFormats[MaxInputFmtCount];
+        NVENCSTATUS		Result = NVENC.nvEncGetInputFormats(InEncoder, NV_ENC_CODEC_H264_GUID, InputFormats, MaxInputFmtCount, &InputFmtCount);
+        if (Result != NV_ENC_SUCCESS)
+        {
+            UE_LOG(LogVideoEncoder, Error, TEXT("Unable to query input formats supported by NvEnc (error: %d)."), Result);
+            return false;
+        }
+
+        for (uint32_t Index = 0; Index < InputFmtCount; ++Index)
+        {
+            switch (InputFormats[Index])
+            {
+            case NV_ENC_BUFFER_FORMAT_IYUV:
+                break;
+            case NV_ENC_BUFFER_FORMAT_NV12:
+                break;
+            case NV_ENC_BUFFER_FORMAT_ARGB:
+                break;
+            case NV_ENC_BUFFER_FORMAT_ABGR:
+            #if PLATFORM_WINDOWS
+                OutSupportedInputFormats.Push(EVideoFrameFormat::D3D11_R8G8B8A8_UNORM);
+                OutSupportedInputFormats.Push(EVideoFrameFormat::D3D12_R8G8B8A8_UNORM);
+            #endif
+            #if WITH_CUDA
+                OutSupportedInputFormats.Push(EVideoFrameFormat::CUDA_R8G8B8A8_UNORM);
+            #endif
+                break;
+            }
+        }
+        return true;
+    }
+
+    static bool GetEncoderInfo(FNVENCCommon& NVENC, FVideoEncoderInfo& EncoderInfo)
+    {
+        bool		bSuccess = true;
+        // create a temporary encoder session
+
+        void* EncoderSession = nullptr;
+
+    #if PLATFORM_WINDOWS
+        TRefCountPtr<ID3D11Device>			EncoderDevice;
+        TRefCountPtr<ID3D11DeviceContext>	EncoderDeviceContext;
+
+        if (!CreateEncoderDevice(EncoderDevice, EncoderDeviceContext))
+        {
+            bSuccess = false;
+        }
+        if ((EncoderSession = CreateEncoderSession(NVENC, EncoderDevice)) == nullptr)
+        {
+            bSuccess = false;
+        }
+    #endif
+
+    #if WITH_CUDA
+        if (!EncoderSession)
+        {
+            EncoderSession = CreateEncoderSession(NVENC, FModuleManager::GetModuleChecked<FCUDAModule>("CUDA").GetCudaContext());
+        }
+    #endif
+
+        EncoderInfo.CodecType = ECodecType::H264;
+        EncoderInfo.MaxWidth = GetEncoderCapability(NVENC, EncoderSession, NV_ENC_CAPS_WIDTH_MAX);
+        EncoderInfo.MaxHeight = GetEncoderCapability(NVENC, EncoderSession, NV_ENC_CAPS_HEIGHT_MAX);
+
+        int LevelMax = GetEncoderCapability(NVENC, EncoderSession, NV_ENC_CAPS_LEVEL_MAX);
+        int LevelMin = GetEncoderCapability(NVENC, EncoderSession, NV_ENC_CAPS_LEVEL_MIN);
+        if (LevelMin > 0 && LevelMax > 0 && LevelMax >= LevelMin)
+        {
+            EncoderInfo.H264.MinLevel = (LevelMin > 9) ? LevelMin : 9;
+            EncoderInfo.H264.MaxLevel = (LevelMax < 9) ? 9 : (LevelMax > NV_ENC_LEVEL_H264_52) ? NV_ENC_LEVEL_H264_52 : LevelMax;
+        }
+        else
+        {
+            UE_LOG(LogVideoEncoder, Error, TEXT("Failed to query min/max h264 level supported by NvEnc (reported min/max=%d/%d)."), LevelMin, LevelMax);
+            bSuccess = false;
+        }
+
+        if (!GetEncoderSupportedProfiles(NVENC, EncoderSession, EncoderInfo.H264.SupportedProfiles) ||
+            !GetEncoderSupportedInputFormats(NVENC, EncoderSession, EncoderInfo.SupportedInputFormats))
+        {
+            bSuccess = false;
+        }
+
+        // destroy encoder session
+        if (EncoderSession)
+        {
+            NVENC.nvEncDestroyEncoder(EncoderSession);
+        }
+
+        return bSuccess;
+    }
 
 } /* namespace AVEncoder */
\ No newline at end of file
