Merge in PixelStreaming, AVEncoder, and SignallingServer changes from TensorWorks local repository

diff --git a/Engine/Source/Runtime/AVEncoder/Private/Encoders/NVENC/NVENC_EncoderH264.h b/Engine/Source/Runtime/AVEncoder/Private/Encoders/NVENC/NVENC_EncoderH264.h
--- a/Engine/Source/Runtime/AVEncoder/Private/Encoders/NVENC/NVENC_EncoderH264.h
+++ b/Engine/Source/Runtime/AVEncoder/Private/Encoders/NVENC/NVENC_EncoderH264.h
@@ -9,146 +9,131 @@
 
 namespace AVEncoder
 {
-	class FVideoEncoderNVENC_H264 : public FVideoEncoder
-	{
-	public:
-		~FVideoEncoderNVENC_H264();
-
-		// query whether or not encoder is supported and available
-		static bool GetIsAvailable(FVideoEncoderInputImpl &InInput, FVideoEncoderInfo &OutEncoderInfo);
-
-		// register encoder with video encoder factory
-		static void Register(FVideoEncoderFactory &InFactory);
-
-		// create instance of video encoder
-		static TUniquePtr<FVideoEncoderNVENC_H264> Create(FVideoEncoderInputImpl &InInput);
-
-		bool Setup(TSharedRef<FVideoEncoderInput> InInput, const FInit &InInit) override;
-		void Encode(const FVideoEncoderInputFrame *InFrame, const FEncodeOptions &InOptions) override;
-		void Flush();
-		void Shutdown() override;
-
-		void UpdateFrameRate(uint32 InMaxFramerate) override;
-		void UpdateLayerBitrate(uint32 InLayerIndex, uint32 InMaxBitRate, uint32 InTargetBitRate) override;
-		void UpdateLayerResolution(uint32 InLayerIndex, uint32 InWidth, uint32 InHeight) override;
-		void UpdateMinQP(int32 minqp) override;
-		void UpdateRateControl(RateControlMode mode) override;
-		void UpdateFillData(bool enable) override;
-
-	protected:
-		FLayerInfo *CreateLayer(uint32 InLayerIndex, const FLayerInfo &InLayerInfo) override;
-		void DestroyLayer(FLayerInfo *InLayerInfo) override;
-
-	private:
-		FVideoEncoderNVENC_H264();
-
-		class FLayer : public FLayerInfo
-		{
-		public:
-			FLayer(uint32 InLayerIndex, const FLayerInfo &InLayerInfo, FVideoEncoderNVENC_H264 &InEncoder);
-			~FLayer();
-
-			bool Setup();
-			bool CreateSession();
-			bool CreateInitialConfig();
-			int GetCapability(NV_ENC_CAPS CapsToQuery) const;
-			FString GetError(NVENCSTATUS ForStatus) const;
-			void Encode(const FVideoEncoderInputFrameImpl *InFrame, const FEncodeOptions &InOptions);
-			void Flush();
-			void Shutdown();
-			void UpdateBitrate(uint32 InMaxBitRate, uint32 InTargetBitRate);
-			void UpdateResolution(uint32 InMaxBitRate, uint32 InTargetBitRate);
-
-			FVideoEncoderNVENC_H264 &Encoder;
-			FNVENCCommon &NVENC;
-			GUID CodecGUID;
-			uint32 LayerIndex;
-			void *NVEncoder = nullptr;
-			NV_ENC_INITIALIZE_PARAMS EncoderInitParams;
-			NV_ENC_CONFIG EncoderConfig;
-			bool bAsyncMode = false;
-			uint32 MaxFramerate = -1;
-			FDateTime LastKeyFrameTime = 0;
-
-			struct FInputOutput
-			{
-				const FVideoEncoderInputFrameImpl *SourceFrame = nullptr;
-
-
-				void*  InputTexture = nullptr;
-				uint32 Width = 0;
-				uint32 Height = 0;
-				uint32 Pitch = 0;
-				NV_ENC_BUFFER_FORMAT BufferFormat = NV_ENC_BUFFER_FORMAT_UNDEFINED;
-				NV_ENC_REGISTERED_PTR RegisteredInput = nullptr;
-				NV_ENC_INPUT_PTR MappedInput = nullptr;
-
-				NV_ENC_OUTPUT_PTR OutputBitstream = nullptr;
-				const void *BitstreamData = nullptr;
-				uint32 BitstreamDataSize = 0;
-				void *CompletionEvent = nullptr;
-				NV_ENC_PIC_TYPE PictureType = NV_ENC_PIC_TYPE_UNKNOWN;
-				uint32 FrameAvgQP = 0;
-				uint64 TimeStamp;
-				FEvent *TriggerOnCompletion;
-
-				FTimespan EncodeStartTs;
-			};
-
-			FInputOutput *GetOrCreateBuffer(const FVideoEncoderInputFrameImpl *InFrame);
-			FInputOutput *CreateBuffer();
-			void DestroyBuffer(FInputOutput *InBuffer);
-			bool RegisterInputTexture(FInputOutput &InBuffer, void *InTexture, FIntPoint TextureSize);
-			bool UnregisterInputTexture(FInputOutput &InBuffer);
-			bool MapInputTexture(FInputOutput &InBuffer);
-			bool UnmapInputTexture(FInputOutput &InBuffer);
-			bool LockOutputBuffer(FInputOutput &InBuffer);
-			bool UnlockOutputBuffer(FInputOutput &InBuffer);
-
-			void ProcessNextPendingFrame();
-			void WaitForNextPendingFrame();
-
-			void CreateResourceDIRECTX(FInputOutput &InBuffer, NV_ENC_REGISTER_RESOURCE &RegisterParam, FIntPoint TextureSize);
-			void CreateResourceCUDAARRAY(FInputOutput &InBuffer, NV_ENC_REGISTER_RESOURCE &RegisterParam, FIntPoint TextureSize);
-
-			TArray<FInputOutput *> CreatedBuffers;
-			TQueue<FInputOutput *> PendingEncodes;
-			FCriticalSection ProtectedWaitingForPending;
-			bool WaitingForPendingActive = false;
-			FThreadSafeBool bUpdateConfig = false;
-		};
-
-		FNVENCCommon &NVENC;
-		EVideoFrameFormat FrameFormat = EVideoFrameFormat::Undefined;
-
-		// TODO (M84FIX) make this neater
+    class FVideoEncoderNVENC_H264 : public FVideoEncoder
+    {
+    public:
+        virtual ~FVideoEncoderNVENC_H264() override;
+
+        // query whether or not encoder is supported and available
+        static bool GetIsAvailable(FVideoEncoderInputImpl &InInput, FVideoEncoderInfo &OutEncoderInfo);
+
+        // register encoder with video encoder factory
+        static void Register(FVideoEncoderFactory &InFactory);
+
+        bool Setup(TSharedRef<FVideoEncoderInput> input, FLayerConfig const& config) override;
+        void Encode(FVideoEncoderInputFrame const* frame, FEncodeOptions const& options) override;
+        void Flush();
+        void Shutdown() override;
+
+    protected:
+        FLayer *CreateLayer(uint32 InLayerIndex, const FLayerConfig &InLayerConfig) override;
+		void DestroyLayer(FLayer* layer) override;
+
+    private:
+        FVideoEncoderNVENC_H264();
+
+        class FNVENCLayer : public FLayer
+        {
+        public:
+            FNVENCLayer(uint32 layerIdx, FLayerConfig const& config, FVideoEncoderNVENC_H264& encoder);
+            ~FNVENCLayer();
+
+            bool Setup();
+            bool CreateSession();
+            bool CreateInitialConfig();
+            int GetCapability(NV_ENC_CAPS CapsToQuery) const;
+            FString GetError(NVENCSTATUS ForStatus) const;
+			void MaybeReconfigure();
+			void UpdateConfig();
+            void Encode(FVideoEncoderInputFrameImpl const* frame, FEncodeOptions const& options);
+            void Flush();
+            void Shutdown();
+            void UpdateBitrate(uint32 InMaxBitRate, uint32 InTargetBitRate);
+            void UpdateResolution(uint32 InMaxBitRate, uint32 InTargetBitRate);
+
+            FVideoEncoderNVENC_H264 &Encoder;
+            FNVENCCommon &NVENC;
+            GUID CodecGUID;
+            uint32 LayerIndex;
+            void *NVEncoder = nullptr;
+            NV_ENC_INITIALIZE_PARAMS EncoderInitParams;
+            NV_ENC_CONFIG EncoderConfig;
+            bool bAsyncMode = false;
+            FDateTime LastKeyFrameTime = 0;
+			bool bForceNextKeyframe = false;
+
+            struct FInputOutput
+            {
+                const FVideoEncoderInputFrameImpl *SourceFrame = nullptr;
+
+                void*  InputTexture = nullptr;
+                uint32 Width = 0;
+                uint32 Height = 0;
+                uint32 Pitch = 0;
+                NV_ENC_BUFFER_FORMAT BufferFormat = NV_ENC_BUFFER_FORMAT_UNDEFINED;
+                NV_ENC_REGISTERED_PTR RegisteredInput = nullptr;
+                NV_ENC_INPUT_PTR MappedInput = nullptr;
+
+                NV_ENC_OUTPUT_PTR OutputBitstream = nullptr;
+                const void *BitstreamData = nullptr;
+                uint32 BitstreamDataSize = 0;
+                void *CompletionEvent = nullptr;
+                NV_ENC_PIC_TYPE PictureType = NV_ENC_PIC_TYPE_UNKNOWN;
+                uint32 FrameAvgQP = 0;
+                uint64 TimeStamp;
+                FEvent *TriggerOnCompletion;
+
+                FTimespan EncodeStartTs;
+            };
+
+            FInputOutput *GetOrCreateBuffer(const FVideoEncoderInputFrameImpl *InFrame);
+            FInputOutput *CreateBuffer();
+            void DestroyBuffer(FInputOutput *InBuffer);
+            bool RegisterInputTexture(FInputOutput &InBuffer, void *InTexture, FIntPoint TextureSize);
+            bool UnregisterInputTexture(FInputOutput &InBuffer);
+            bool MapInputTexture(FInputOutput &InBuffer);
+            bool UnmapInputTexture(FInputOutput &InBuffer);
+            bool LockOutputBuffer(FInputOutput &InBuffer);
+            bool UnlockOutputBuffer(FInputOutput &InBuffer);
+
+            void ProcessNextPendingFrame();
+            void WaitForNextPendingFrame();
+
+            void CreateResourceDIRECTX(FInputOutput &InBuffer, NV_ENC_REGISTER_RESOURCE &RegisterParam, FIntPoint TextureSize);
+            void CreateResourceCUDAARRAY(FInputOutput &InBuffer, NV_ENC_REGISTER_RESOURCE &RegisterParam, FIntPoint TextureSize);
+
+            TArray<FInputOutput *> CreatedBuffers;
+            TQueue<FInputOutput *> PendingEncodes;
+            FCriticalSection ProtectedWaitingForPending;
+            bool WaitingForPendingActive = false;
+        };
+
+        FNVENCCommon &NVENC;
+        EVideoFrameFormat FrameFormat = EVideoFrameFormat::Undefined;
+
+        // TODO (M84FIX) make this neater
 #if PLATFORM_WINDOWS
-		TRefCountPtr<ID3D11Device> EncoderDevice;
+        TRefCountPtr<ID3D11Device> EncoderDevice;
 #elif WITH_CUDA
-		CUcontext EncoderDevice;
+        CUcontext EncoderDevice;
 #else
-		void* EncoderDevice;
+        void* EncoderDevice;
 #endif
 
-		uint32 MaxFramerate = 0;
-		int32 MinQP = -1;
-		RateControlMode RateMode = RateControlMode::CBR;
-		bool FillData = false;
+        // event thread for nvenc async (windows only)
+        void OnEvent(void *InEvent, TUniqueFunction<void()> &&InCallback);
+        void StartEventThread();
+        void StopEventThread();
+        void EventLoop();
 
-		// event thread for nvenc async (windows only)
-		void OnEvent(void *InEvent, TUniqueFunction<void()> &&InCallback);
-		void StartEventThread();
-		void StopEventThread();
-		void EventLoop();
+        TUniquePtr<FThread> EventThread;
+        FCriticalSection ProtectEventThread;
+        bool bExitEventThread = false;
 
-		TUniquePtr<FThread> EventThread;
-		FCriticalSection ProtectEventThread;
-		bool bExitEventThread = false;
+        void* EventThreadCheckEvent = nullptr;
 
-		void *EventThreadCheckEvent = nullptr;
-
-		using FWaitForEvent = TPair<void *, TUniqueFunction<void()>>;
-		TArray<FWaitForEvent> EventThreadWaitingFor;
-	};
+        using FWaitForEvent = TPair<void *, TUniqueFunction<void()>>;
+        TArray<FWaitForEvent> EventThreadWaitingFor;
+    };
 
 }
\ No newline at end of file
diff --git a/Engine/Source/Runtime/AVEncoder/Private/VideoEncoder.cpp b/Engine/Source/Runtime/AVEncoder/Private/VideoEncoder.cpp
--- a/Engine/Source/Runtime/AVEncoder/Private/VideoEncoder.cpp
+++ b/Engine/Source/Runtime/AVEncoder/Private/VideoEncoder.cpp
@@ -10,62 +10,55 @@ namespace AVEncoder
 
 FVideoEncoder::~FVideoEncoder()
 {
-	for (int32 LayerIndex = LayerInfos.Num(); LayerIndex-- > 0; )
+	// calling virtual methods from the destructor is bad so 'DestroyLayer' was removed from here
+	// its now up to the derived classes to clean up
+	for (auto&& layer : Layers)
 	{
-		DestroyLayer(LayerInfos[LayerIndex]);
+		if (layer)
+		{
+			UE_LOG(LogVideoEncoder, Warning, TEXT("Encoder layer may not have been cleaned up!"));
+			break;
+		}
 	}
 }
 
-bool FVideoEncoder::Setup(TSharedRef<FVideoEncoderInput> InInput, const FInit& InInit)
+bool FVideoEncoder::AddLayer(FLayerConfig const& config)
 {
-	bool		bSuccess = false;
-	if (GetNumLayers() == 0)
-	{
-		bSuccess = AddLayer(InInit);
-	}
-	return bSuccess;
-}
+    if (GetNumLayers() >= GetMaxLayers())
+    {
+        UE_LOG(LogVideoEncoder, Error, TEXT("Encoder does not support more than %d layers."), GetMaxLayers());
+        return false;
+    }
 
-void FVideoEncoder::Shutdown()
-{
+    auto const newLayer = CreateLayer(Layers.Num(), config);
+    if (nullptr != newLayer)
+    {
+        Layers.Push(newLayer);
+        return true;
+    }
+    return false;
 }
 
-bool FVideoEncoder::AddLayer(const FLayerConfig& InLayerConfig)
+FVideoEncoder::FLayerConfig FVideoEncoder::GetLayerConfig(uint32 layerIdx) const
 {
-	if (static_cast<uint32>(LayerInfos.Num()) >= GetMaxLayers())
-	{
-		UE_LOG(LogVideoEncoder, Error, TEXT("Encoder does not support more than %d layers."), GetMaxLayers());
-		return false;
-	}
-	FLayerInfo		LayerInfo(InLayerConfig);
-	FLayerInfo* NewLayerInfo = CreateLayer(LayerInfos.Num(), LayerInfo);
-	if (NewLayerInfo)
+	if (layerIdx >= GetNumLayers())
 	{
-		LayerInfos.Push(NewLayerInfo);
+		UE_LOG(LogVideoEncoder, Error, TEXT("FVideoEncoder::GetLayerConfig: Layer index %d out of range."), layerIdx);
+		return {};
 	}
-	return NewLayerInfo != nullptr;
-}
 
-FVideoEncoder::FLayerInfo* FVideoEncoder::CreateLayer(uint32 InLayerIndex, const FLayerInfo& InLayerInfo)
-{
-	return new FLayerInfo(InLayerInfo);
+	return Layers[layerIdx]->GetConfig();
 }
 
-void FVideoEncoder::DestroyLayer(FLayerInfo* InLayerInfo)
+void FVideoEncoder::UpdateLayerConfig(uint32 layerIdx, FLayerConfig const& config)
 {
-	delete InLayerInfo;
-}
+    if (layerIdx >= GetNumLayers())
+    {
+        UE_LOG(LogVideoEncoder, Error, TEXT("FVideoEncoder::UpdateLayerConfig: Layer index %d out of range."), layerIdx);
+        return;
+    }
 
-FVideoEncoder::FLayerInfo::FLayerInfo(const FLayerConfig& InLayerConfig)
-	: Width(InLayerConfig.Width)
-	, Height(InLayerConfig.Height)
-	, MaxBitrate(InLayerConfig.MaxBitrate)
-	, TargetBitrate(InLayerConfig.TargetBitrate)
-	, QPMax(InLayerConfig.QPMax)
-	, QPMin(InLayerConfig.QPMin)
-	, RateControlMode(InLayerConfig.RateControlMode)
-	, FillData(InLayerConfig.FillData)
-{
+    Layers[layerIdx]->UpdateConfig(config);
 }
 
 } /* namespace AVEncoder */
diff --git a/Engine/Source/Runtime/AVEncoder/Private/VideoEncoderFactory.cpp b/Engine/Source/Runtime/AVEncoder/Private/VideoEncoderFactory.cpp
--- a/Engine/Source/Runtime/AVEncoder/Private/VideoEncoderFactory.cpp
+++ b/Engine/Source/Runtime/AVEncoder/Private/VideoEncoderFactory.cpp
@@ -8,6 +8,8 @@
 #include "Encoders/NVENC/NVENC_EncoderH264.h"
 #endif
 
+#include "Encoders/Amf/Amf_EncoderH264.h"
+
 #include "Encoders/VideoEncoderH264_Dummy.h"
 
 namespace AVEncoder
@@ -71,6 +73,8 @@ void FVideoEncoderFactory::RegisterDefaultCodecs()
 #if PLATFORM_WINDOWS || (PLATFORM_LINUX && WITH_CUDA)
 	FVideoEncoderNVENC_H264::Register(*this);
 #endif
+	
+	FVideoEncoderAmf_H264::Register(*this);
 
 #if defined(AVENCODER_VIDEO_ENCODER_AVAILABLE_H264_DUMMY)
 	FVideoEncoderH264_Dummy::Register(*this);
@@ -102,7 +106,7 @@ bool FVideoEncoderFactory::HasEncoderForCodec(ECodecType CodecType) const
 	return false;
 }
 
-TUniquePtr<FVideoEncoder> FVideoEncoderFactory::Create(uint32 InID, const FVideoEncoder::FInit& InInit)
+TUniquePtr<FVideoEncoder> FVideoEncoderFactory::Create(uint32 InID, const FVideoEncoder::FLayerConfig& config)
 {
 	// HACK (M84FIX) create encoder without a ready FVideoEncoderInput
 	TUniquePtr<FVideoEncoder>	Result;
@@ -115,9 +119,9 @@ TUniquePtr<FVideoEncoder> FVideoEncoderFactory::Create(uint32 InID, const FVideo
 			// HACK (M84FIX) work with other RHI
 			if (GDynamicRHI->GetName() == FString("D3D11"))
 			{
-				TSharedRef<FVideoEncoderInput> Input = FVideoEncoderInput::CreateForD3D11(GDynamicRHI->RHIGetNativeDevice(), InInit.Width, InInit.Height).ToSharedRef();
+				TSharedRef<FVideoEncoderInput> Input = FVideoEncoderInput::CreateForD3D11(GDynamicRHI->RHIGetNativeDevice(), config.Width, config.Height).ToSharedRef();
 
-				if (Result && !Result->Setup(Input, InInit))
+				if (Result && !Result->Setup(Input, config))
 				{
 					Result.Reset();
 				}
@@ -125,9 +129,9 @@ TUniquePtr<FVideoEncoder> FVideoEncoderFactory::Create(uint32 InID, const FVideo
 			}
 			else if (GDynamicRHI->GetName() == FString("D3D12"))
 			{
-				TSharedRef<FVideoEncoderInput> Input = FVideoEncoderInput::CreateForD3D12(GDynamicRHI->RHIGetNativeDevice(), InInit.Width, InInit.Height).ToSharedRef();
+				TSharedRef<FVideoEncoderInput> Input = FVideoEncoderInput::CreateForD3D12(GDynamicRHI->RHIGetNativeDevice(), config.Width, config.Height).ToSharedRef();
 
-				if (Result && !Result->Setup(Input, InInit))
+				if (Result && !Result->Setup(Input, config))
 				{
 					Result.Reset();
 				}
@@ -138,7 +142,7 @@ TUniquePtr<FVideoEncoder> FVideoEncoderFactory::Create(uint32 InID, const FVideo
 	return Result;
 }
 
-TUniquePtr<FVideoEncoder> FVideoEncoderFactory::Create(uint32 InID, TSharedPtr<FVideoEncoderInput> InInput, const FVideoEncoder::FInit& InInit)
+TUniquePtr<FVideoEncoder> FVideoEncoderFactory::Create(uint32 InID, TSharedPtr<FVideoEncoderInput> InInput, const FVideoEncoder::FLayerConfig& config)
 {
 	TUniquePtr<FVideoEncoder>		Result;
 	if (InInput)
@@ -149,7 +153,7 @@ TUniquePtr<FVideoEncoder> FVideoEncoderFactory::Create(uint32 InID, TSharedPtr<F
 			if (AvailableEncoders[Index].ID == InID)
 			{
 				Result = CreateEncoders[Index]();
-				if (Result && !Result->Setup(MoveTemp(Input), InInit))
+				if (Result && !Result->Setup(MoveTemp(Input), config))
 				{
 					Result.Reset();
 				}
diff --git a/Engine/Source/Runtime/AVEncoder/Private/VideoEncoderInput.cpp b/Engine/Source/Runtime/AVEncoder/Private/VideoEncoderInput.cpp
--- a/Engine/Source/Runtime/AVEncoder/Private/VideoEncoderInput.cpp
+++ b/Engine/Source/Runtime/AVEncoder/Private/VideoEncoderInput.cpp
@@ -7,33 +7,10 @@
 #include "AVEncoderDebug.h"
 #include "Misc/Paths.h"
 #include "VideoCommon.h"
+#include "VulkanRHIPrivate.h"
 
 #if PLATFORM_WINDOWS
-
-// // Disable macro redefinition warning for compatibility with Windows SDK 8+
-// #pragma warning(push)
-// #pragma warning(disable : 4005)	// macro redefinition
-
-// #include "Windows/AllowWindowsPlatformTypes.h"
-// #include "Windows/PreWindowsApi.h"
-// 	#include <d3d11.h>
-// 	#include <mftransform.h>
-// 	#include <mfapi.h>
-// 	#include <mferror.h>
-// 	#include <mfidl.h>
-// 	#include <codecapi.h>
-// 	#include <shlwapi.h>
-// 	#include <mfreadwrite.h>
-// 	#include <d3d11_1.h>
-// 	#include <d3d12.h>
-// 	#include <dxgi1_4.h>
-// #include "Windows/PostWindowsApi.h"
-// #include "Windows/HideWindowsPlatformTypes.h"
 #include "MicrosoftCommon.h"
-#endif /* PLATFORM_WINDOWS */
-
-#if WITH_CUDA
-#include "CudaModule.h"
 #endif
 
 namespace AVEncoder
@@ -565,6 +542,11 @@ CUcontext FVideoEncoderInputImpl::GetCUDAEncoderContext() const
 
 #endif
 
+VkDevice FVideoEncoderInputImpl::GetVulkanDevice() const
+{
+	return FrameInfoVulkan.VulkanDevice;
+}
+
 // *** FVideoEncoderInputFrame ********************************************************************
 
 FVideoEncoderInputFrame::FVideoEncoderInputFrame()
@@ -580,7 +562,8 @@ FVideoEncoderInputFrame::FVideoEncoderInputFrame()
 }
 
 FVideoEncoderInputFrame::FVideoEncoderInputFrame(const FVideoEncoderInputFrame& CloneFrom)
-	: FrameID(CloneFrom.FrameID)
+	: PTS(CloneFrom.PTS)
+	, FrameID(CloneFrom.FrameID)
 	, NumReferences(0)
 	, Format(CloneFrom.Format)
 	, Width(CloneFrom.Width)
@@ -612,6 +595,7 @@ FVideoEncoderInputFrame::~FVideoEncoderInputFrame()
 		YUV420P.Data[0] = YUV420P.Data[1] = YUV420P.Data[2] = nullptr;
 		bFreeYUV420PData = false;
 	}
+
 #if PLATFORM_WINDOWS
 	if (D3D11.EncoderTexture)
 	{
@@ -663,6 +647,20 @@ FVideoEncoderInputFrame::~FVideoEncoderInputFrame()
 
 }
 
+void FVideoEncoderInputFrame::AllocateYUV420P()
+{
+	if (!bFreeYUV420PData)
+	{
+		YUV420P.StrideY = Width;
+		YUV420P.StrideU = (Width + 1) / 2;
+		YUV420P.StrideV = (Width + 1) / 2;;
+		YUV420P.Data[0] = new uint8[Height * YUV420P.StrideY];
+		YUV420P.Data[1] = new uint8[(Height + 1) / 2 * YUV420P.StrideU];
+		YUV420P.Data[2] = new uint8[(Height + 1) / 2 * YUV420P.StrideV];
+		bFreeYUV420PData = true;
+	}
+}
+
 void FVideoEncoderInputFrame::SetYUV420P(const uint8* InDataY, const uint8* InDataU, const uint8* InDataV, uint32 InStrideY, uint32 InStrideU, uint32 InStrideV)
 {
 	if (Format == EVideoFrameFormat::YUV420P)
@@ -799,6 +797,18 @@ void FVideoEncoderInputFrame::SetTexture(CUarray InTexture, FReleaseCUDATextureC
 #endif
 
 
+void FVideoEncoderInputFrame::SetTexture(VkImage InTexture, FReleaseVulkanTextureCallback InOnReleaseTexture)
+{
+	if (Format == EVideoFrameFormat::VULKAN_R8G8B8A8_UNORM)
+	{
+		Vulkan.EncoderTexture = InTexture;
+		OnReleaseVulkanTexture = InOnReleaseTexture;
+		if (!Vulkan.EncoderTexture)
+		{
+			UE_LOG(LogVideoEncoder, Warning, TEXT("SetTexture | Vulkan VkImage is null"));
+		}
+	}
+}
 
 // *** FVideoEncoderInputFrameImpl ****************************************************************
 
diff --git a/Engine/Source/Runtime/AVEncoder/Private/VideoEncoderInputImpl.h b/Engine/Source/Runtime/AVEncoder/Private/VideoEncoderInputImpl.h
--- a/Engine/Source/Runtime/AVEncoder/Private/VideoEncoderInputImpl.h
+++ b/Engine/Source/Runtime/AVEncoder/Private/VideoEncoderInputImpl.h
@@ -7,14 +7,16 @@
 #include "Misc/ScopeLock.h"
 #include "Templates/RefCounting.h"
 
+#if WITH_CUDA
+#include "CudaModule.h"
+#endif
+
 // HACK (M84FIX) need to break these dependencies
 #if PLATFORM_WINDOWS
 struct ID3D11DeviceContext;
 #endif
 
-#if WITH_CUDA
-#include "CudaModule.h"
-#endif
+struct VkDevice_T;
 
 namespace AVEncoder
 {
@@ -45,6 +47,9 @@ public:
 	// set up an encoder that encodes a CUArray in the CUDA context
 	bool SetupForCUDA(void* InApplicationContext, uint32 InWidth, uint32 InHeight);
 
+	// set up an encoder that encodes a VkImage in the context of a VkDevice
+	bool SetupForVulkan(void* InApplicationContext, uint32 InWidth, uint32 InHeight);
+
 	// --- available encoders
 
 	// get a list of supported video encoders
@@ -80,6 +85,7 @@ public:
 	CUcontext GetCUDAEncoderContext() const;
 #endif
 
+	VkDevice_T* GetVulkanDevice() const;
 
 private:
 
@@ -122,6 +128,11 @@ private:
 	}								FrameInfoCUDA;
 #endif
 
+	struct FFrameInfoVulkan
+	{
+		VkDevice_T*					VulkanDevice;
+	}								FrameInfoVulkan;
+
 	mutable FCriticalSection				ProtectFrames;
 	TQueue<FVideoEncoderInputFrameImpl*>	AvailableFrames;
 	TArray<FVideoEncoderInputFrameImpl*>	ActiveFrames;
diff --git a/Engine/Source/Runtime/AVEncoder/Public/IMFSampleWrapper.h b/Engine/Source/Runtime/AVEncoder/Public/IMFSampleWrapper.h
--- a/Engine/Source/Runtime/AVEncoder/Public/IMFSampleWrapper.h
+++ b/Engine/Source/Runtime/AVEncoder/Public/IMFSampleWrapper.h
@@ -8,12 +8,6 @@
 
 #if PLATFORM_WINDOWS
 
-/*#include "Windows/AllowWindowsPlatformTypes.h"
-#include "Windows/PreWindowsApi.h"
-	#include <mftransform.h>
-#include "Windows/PostWindowsApi.h"
-#include "Windows/HideWindowsPlatformTypes.h"*/
-
 namespace AVEncoder
 {
 
diff --git a/Engine/Source/Runtime/AVEncoder/Public/VideoCommon.h b/Engine/Source/Runtime/AVEncoder/Public/VideoCommon.h
--- a/Engine/Source/Runtime/AVEncoder/Public/VideoCommon.h
+++ b/Engine/Source/Runtime/AVEncoder/Public/VideoCommon.h
@@ -46,7 +46,18 @@ namespace AVEncoder
 		YUV420P,				// Planar YUV420 format in CPU memory
 		D3D11_R8G8B8A8_UNORM,	//
 		D3D12_R8G8B8A8_UNORM,	//
-		CUDA_R8G8B8A8_UNORM
+		CUDA_R8G8B8A8_UNORM,
+		VULKAN_R8G8B8A8_UNORM,
+	};
+
+	enum class EH264Profile
+	{
+		UNKNOWN,
+		CONSTRAINED_BASELINE,
+		BASELINE,
+		MAIN,
+		CONSTRAINED_HIGH,
+		HIGH,
 	};
 
 	inline FString ToString(EVideoFrameFormat Format)
@@ -61,6 +72,8 @@ namespace AVEncoder
 			return FString("EVideoFrameFormat::D3D12_R8G8B8A8_UNORM");
 		case EVideoFrameFormat::CUDA_R8G8B8A8_UNORM:
 			return FString("EVideoFrameFormat::CUDA_R8G8B8A8_UNORM");
+		case EVideoFrameFormat::VULKAN_R8G8B8A8_UNORM:
+			return FString("EVideoFrameFormat::VULKAN_R8G8B8A8_UNORM");
 		case EVideoFrameFormat::Undefined:
 		default:
 			return FString("EVideoFrameFormat::Undefined");
@@ -75,7 +88,7 @@ namespace AVEncoder
 		VP8,
 	};
 
-
+	// TODO: make enums
 	const uint32 H264Profile_ConstrainedBaseline = 1 << 0;
 	const uint32 H264Profile_Baseline = 1 << 1;
 	const uint32 H264Profile_Main = 1 << 2;
diff --git a/Engine/Source/Runtime/AVEncoder/Public/VideoEncoder.h b/Engine/Source/Runtime/AVEncoder/Public/VideoEncoder.h
--- a/Engine/Source/Runtime/AVEncoder/Public/VideoEncoder.h
+++ b/Engine/Source/Runtime/AVEncoder/Public/VideoEncoder.h
@@ -5,112 +5,105 @@
 #include "CoreMinimal.h"
 #include "VideoEncoderInput.h"
 #include "Misc/FrameRate.h"
+#include "Misc/ScopeLock.h"
 
 namespace AVEncoder
 {
-	class FVideoEncoder
-	{
-	public:
-		virtual ~FVideoEncoder();
-
-		enum class RateControlMode
-		{
-			UNKNOWN,
-			CONSTQP,
-			VBR,
-			CBR,
-		};
-
-		struct FLayerConfig
-		{
-			uint32			Width = 0;
-			uint32			Height = 0;
-			uint32			MaxBitrate = 0;
-			uint32			TargetBitrate = 0;
-			uint32			QPMax = 0;
-			int32			QPMin = -1;
-			RateControlMode RateControlMode = RateControlMode::CBR;
-			bool			FillData = false;
-		};
-
-		struct FInit : public FLayerConfig
-		{
-			uint32			MaxFramerate = 0;
-			FFrameRate		TimeBase;
-		};
-
-		// --- setup
-
-		virtual bool Setup(TSharedRef<FVideoEncoderInput> InInput, const FInit& InInit);
-		virtual void Shutdown();
-
-		// --- properties
-
-		// --- layers
-
-
-		// get the max. number of supported layers - the original resolution counts as a layer so at least one layer is supported.
-		virtual uint32 GetMaxLayers() const { return 1; }
-		// add a layer to encode - each consecutive layer must be smaller than the previous one
-		virtual bool AddLayer(const FLayerConfig& InLayerConfig);
-		// get the current number of layers
-		uint32 GetNumLayers() const { return LayerInfos.Num(); }
-
-		uint32 GetWidth(uint32 InLayerIndex) const { return (InLayerIndex < static_cast<uint32>(LayerInfos.Num())) ? LayerInfos[InLayerIndex]->Width : 0; }
-		uint32 GetHeight(uint32 InLayerIndex) const { return (InLayerIndex < static_cast<uint32>(LayerInfos.Num())) ? LayerInfos[InLayerIndex]->Height : 0; }
-
-		virtual void UpdateFrameRate(uint32 InMaxFramerate) {}
-		virtual void UpdateLayerBitrate(uint32 InLayerIndex, uint32 InMaxBitRate, uint32 InTargetBitRate) {}
-		virtual void UpdateLayerResolution(uint32 InLayerIndex, uint32 InWidth, uint32 InHeight) {}
-		virtual void UpdateMinQP(int32 minqp) {}
-		virtual void UpdateRateControl(RateControlMode mode) {}
-		virtual void UpdateFillData(bool enable) {}
-
-		// --- input
-
-		// new packet callback prototype void(uint32 LayerIndex, const FCodecPacket& Packet)
-		using OnFrameEncodedCallback = TFunction<void(const FVideoEncoderInputFrame* /* InCompletedFrame */)>;
-
-		struct FEncodeOptions
-		{
-			bool					bForceKeyFrame = false;
-			OnFrameEncodedCallback	OnFrameEncoded;
-		};
-
-		virtual void Encode(const FVideoEncoderInputFrame* InFrame, const FEncodeOptions& InOptions) = 0;
-
-		// --- output
-
-		// new packet callback prototype void(uint32 LayerIndex, const FCodecPacket& Packet)
-		using OnEncodedPacketCallback = TFunction<void(uint32 /* LayerIndex */, const FVideoEncoderInputFrame* /* Frame */, const FCodecPacket& /* Packet */)>;
-
-		void SetOnEncodedPacket(OnEncodedPacketCallback InCallback) { OnEncodedPacket = MoveTemp(InCallback); }
-		void ClearOnEncodedPacket() { OnEncodedPacket = nullptr; }
-
-	protected:
-		FVideoEncoder() = default;
-
-		class FLayerInfo
-		{
-		public:
-			explicit FLayerInfo(const FLayerConfig& InLayerConfig);
-			virtual ~FLayerInfo() = default;
-
-			uint32			Width;
-			uint32			Height;
-			uint32			MaxBitrate;
-			uint32			TargetBitrate;
-			uint32			QPMax;
-			int32			QPMin;
-			RateControlMode RateControlMode;
-			bool			FillData;
-		};
-
-		TArray<FLayerInfo*>		LayerInfos;
-
-		virtual FLayerInfo* CreateLayer(uint32 InLayerIndex, const FLayerInfo& InLayerInfo);
-		virtual void DestroyLayer(FLayerInfo* InLayerInfo);
-
-		OnEncodedPacketCallback			OnEncodedPacket;
-	};
+    class AVENCODER_API FVideoEncoder
+    {
+    public:
+        enum class RateControlMode { UNKNOWN, CONSTQP, VBR, CBR };
+        enum class MultipassMode { UNKNOWN, DISABLED, QUARTER, FULL };
+
+        struct FLayerConfig
+        {
+            uint32			Width = 0;
+            uint32			Height = 0;
+            uint32			MaxFramerate = 0;
+            int32			MaxBitrate = 0;
+            int32			TargetBitrate = 0;
+            int32			QPMax = -1;
+            int32			QPMin = -1;
+            RateControlMode RateControlMode = RateControlMode::CBR;
+            MultipassMode	MultipassMode = MultipassMode::FULL;
+            bool			FillData = false;
+
+			bool operator==(FLayerConfig const& other) const
+			{
+				return Width == other.Width
+					&& Height == other.Height
+					&& MaxFramerate == other.MaxFramerate
+					&& MaxBitrate == other.MaxBitrate
+					&& TargetBitrate == other.TargetBitrate
+					&& QPMax == other.QPMax
+					&& QPMin == other.QPMin
+					&& RateControlMode == other.RateControlMode
+					&& MultipassMode == other.MultipassMode
+					&& FillData == other.FillData;
+			}
+
+			bool operator!=(FLayerConfig const& other) const
+			{
+				return !(*this == other);
+			}
+        };
+
+        virtual ~FVideoEncoder();
+
+		virtual bool Setup(TSharedRef<FVideoEncoderInput> input, FLayerConfig const& config) { return false; }
+		virtual void Shutdown() {}
+
+        virtual bool AddLayer(FLayerConfig const& config);
+        uint32 GetNumLayers() const { return static_cast<uint32>(Layers.Num()); }
+        virtual uint32 GetMaxLayers() const { return 1; }
+
+		FLayerConfig GetLayerConfig(uint32 layerIdx) const;
+        void UpdateLayerConfig(uint32 layerIdx, FLayerConfig const& config);
+
+        using OnFrameEncodedCallback = TFunction<void(const FVideoEncoderInputFrame* /* InCompletedFrame */)>;
+        using OnEncodedPacketCallback = TFunction<void(uint32 /* LayerIndex */, const FVideoEncoderInputFrame* /* Frame */, const FCodecPacket& /* Packet */)>;
+
+        struct FEncodeOptions
+        {
+            bool					bForceKeyFrame = false;
+            OnFrameEncodedCallback	OnFrameEncoded;
+        };
+
+        void SetOnEncodedPacket(OnEncodedPacketCallback callback) { OnEncodedPacket = MoveTemp(callback); }
+        void ClearOnEncodedPacket() { OnEncodedPacket = nullptr; }
+
+		virtual void Encode(FVideoEncoderInputFrame const* frame, FEncodeOptions const& options) {}
+
+    protected:
+        FVideoEncoder() = default;
+
+        class FLayer
+        {
+        public:
+            explicit FLayer(FLayerConfig const& layerConfig)
+                : CurrentConfig(layerConfig)
+                , NeedsReconfigure(false)
+            {}
+            virtual ~FLayer() = default;
+
+			FLayerConfig const& GetConfig() const { return CurrentConfig; }
+			void UpdateConfig(FLayerConfig const& config)
+			{
+				FScopeLock lock(&ConfigMutex);
+				CurrentConfig = config;
+				NeedsReconfigure = true;
+			}
+
+		protected:
+			FCriticalSection	ConfigMutex;
+            FLayerConfig		CurrentConfig;
+            bool				NeedsReconfigure;
+        };
+
+		virtual FLayer* CreateLayer(uint32 layerIdx, FLayerConfig const& config) { return nullptr; }
+		virtual void DestroyLayer(FLayer* layer) {};
+
+        TArray<FLayer*>			Layers;
+        OnEncodedPacketCallback	OnEncodedPacket;
+    };
 }
\ No newline at end of file
diff --git a/Engine/Source/Runtime/AVEncoder/Public/VideoEncoderFactory.h b/Engine/Source/Runtime/AVEncoder/Public/VideoEncoderFactory.h
--- a/Engine/Source/Runtime/AVEncoder/Public/VideoEncoderFactory.h
+++ b/Engine/Source/Runtime/AVEncoder/Public/VideoEncoderFactory.h
@@ -32,8 +32,8 @@ namespace AVEncoder
 
 		// --- encoder creation
 
-		TUniquePtr<FVideoEncoder> Create(uint32 InID, const FVideoEncoder::FInit& InInit);
-		TUniquePtr<FVideoEncoder> Create(uint32 InID, TSharedPtr<FVideoEncoderInput> InInput, const FVideoEncoder::FInit& InInit);
+		TUniquePtr<FVideoEncoder> Create(uint32 InID, const FVideoEncoder::FLayerConfig& config);
+		TUniquePtr<FVideoEncoder> Create(uint32 InID, TSharedPtr<FVideoEncoderInput> InInput, const FVideoEncoder::FLayerConfig& config);
 
 	private:
 		FVideoEncoderFactory() = default;
diff --git a/Engine/Source/Runtime/AVEncoder/Public/VideoEncoderInput.h b/Engine/Source/Runtime/AVEncoder/Public/VideoEncoderInput.h
--- a/Engine/Source/Runtime/AVEncoder/Public/VideoEncoderInput.h
+++ b/Engine/Source/Runtime/AVEncoder/Public/VideoEncoderInput.h
@@ -16,11 +16,14 @@ struct ID3D12Device;
 struct ID3D12Resource;
 #endif
 
+// vulkan forward declaration
+struct VkImage_T;
+struct VkDevice_T;
+
 namespace AVEncoder
 {
 	class FVideoEncoderInputFrame;
 
-
 	class AVENCODER_API FVideoEncoderInput
 	{
 	public:
@@ -31,12 +34,16 @@ namespace AVEncoder
 		// create input for an encoder that encodes a D3D11 texture
 		static TSharedPtr<FVideoEncoderInput> CreateForD3D11(void* InApplicationD3D11Device, uint32 InWidth, uint32 InHeight, bool isResizable = false);
 
+		// TODO (M84FIX) AMF can work with this but also can handle raw D3D12 textures we should add support for that too 
 		// create input for an encoder that encodes a D3D12 texture in the context of a D3D11 device (i.e. nvenc)
 		static TSharedPtr<FVideoEncoderInput> CreateForD3D12(void* InApplicationD3D12Device, uint32 InWidth, uint32 InHeight, bool isResizable = false);
 
-		// create input for an encoder that encodes a Vulkan External Texture in the context of a CUDA context (for nvenc)
+		// create input for an encoder that encodes a CUarray in the context of a CUcontext (i.e. nvenc)
 		static TSharedPtr<FVideoEncoderInput> CreateForCUDA(void* InApplicationCudaContext, uint32 InWidth, uint32 InHeight, bool isResizable = false);
-		
+
+		// create input for an encoder that encodes a VkImage in the context of a VkDevice (i.e. Amf)
+		static TSharedPtr<FVideoEncoderInput> CreateForVulkan(void* InApplicationVulkanDevice, uint32 InWidth, uint32 InHeight, bool isResizable = false);
+
 		// --- properties
 		virtual void SetResolution(uint32 InWidth, uint32 InHeight);
 
@@ -113,7 +120,8 @@ namespace AVEncoder
 		TFunction<void()> OnTextureEncode;
 
 		int64			PTS = 0;		// presentation time stamp (within TimeBase)
-		
+		int64			NTP = 0;		// capture time stamp
+
 		// --- YUV420P
 
 		struct FYUV420P
@@ -124,8 +132,14 @@ namespace AVEncoder
 			uint32				StrideV = 0;
 		};
 
-		const FYUV420P& GetYUV420P() const { return YUV420P; }
+		void AllocateYUV420P();
+		const FYUV420P& GetYUV420P() const 
+		{	
+			return YUV420P; 
+		}
+		
 		FYUV420P& GetYUV420P() { return YUV420P; }
+
 		void SetYUV420P(const uint8* InDataY, const uint8* InDataU, const uint8* InDataV, uint32 InStrideY, uint32 InStrideU, uint32 InStrideV);
 
 #if PLATFORM_WINDOWS
@@ -185,6 +199,21 @@ namespace AVEncoder
 
 #endif // WITH_CUDA
 
+		// --- Vulkan
+		struct FVulkan
+		{
+			VkImage_T*		EncoderTexture;
+			VkDevice_T*		EncoderDevice;
+		};
+
+		const FVulkan& GetVulkan() const { return Vulkan; }
+		FVulkan& GetVulkan() { return Vulkan; }
+
+		// the callback type used to create a registered encoder
+		using FReleaseVulkanTextureCallback = TFunction<void(VkImage_T*)>;
+
+		void SetTexture(VkImage_T* InTexture, FReleaseVulkanTextureCallback InOnReleaseTexture);
+
 	protected:
 		FVideoEncoderInputFrame();
 		explicit FVideoEncoderInputFrame(const FVideoEncoderInputFrame& CloneFrom);
@@ -197,6 +226,7 @@ namespace AVEncoder
 		uint32									Height;
 		FYUV420P								YUV420P;
 		bool									bFreeYUV420PData;
+
 #if PLATFORM_WINDOWS
 		FD3D11									D3D11;
 		FReleaseD3D11TextureCallback			OnReleaseD3D11Texture;
@@ -209,6 +239,8 @@ namespace AVEncoder
 		FReleaseCUDATextureCallback				OnReleaseCUDATexture;
 #endif
 
+		FVulkan									Vulkan;
+		FReleaseVulkanTextureCallback			OnReleaseVulkanTexture;
 	};
 
 
diff --git a/Engine/Source/Runtime/GameplayMediaEncoder/Private/GameplayMediaEncoder.cpp b/Engine/Source/Runtime/GameplayMediaEncoder/Private/GameplayMediaEncoder.cpp
--- a/Engine/Source/Runtime/GameplayMediaEncoder/Private/GameplayMediaEncoder.cpp
+++ b/Engine/Source/Runtime/GameplayMediaEncoder/Private/GameplayMediaEncoder.cpp
@@ -20,8 +20,6 @@
 
 #include "AudioEncoderFactory.h"
 
-#include "VideoCommon.h"
-
 DEFINE_LOG_CATEGORY(GameplayMediaEncoder);
 CSV_DEFINE_CATEGORY(GameplayMediaEncoder, true);
 
@@ -232,7 +230,7 @@ bool FGameplayMediaEncoder::Initialize()
 	VideoConfig.Bitrate = FMath::Clamp(VideoConfig.Bitrate, (uint32)MinVideoBitrate, (uint32)MaxVideoBitrate);
 
 
-	AVEncoder::FVideoEncoder::FInit videoInit;
+	AVEncoder::FVideoEncoder::FLayerConfig videoInit;
 	videoInit.Width = VideoConfig.Width;
 	videoInit.Height = VideoConfig.Height;
 	videoInit.MaxBitrate = MaxVideoBitrate;
@@ -610,19 +608,25 @@ void FGameplayMediaEncoder::SetVideoFramerate(uint32 Framerate)
 
 bool FGameplayMediaEncoder::ChangeVideoConfig()
 {
-	if (bChangeBitrate)
+	if (bChangeBitrate || bChangeFramerate)
 	{
-		VideoEncoder->UpdateLayerBitrate(0, MaxVideoBitrate, NewVideoBitrate);
-		bChangeBitrate = false;
-	}
+		auto config = VideoEncoder->GetLayerConfig(0);
 
-	if (bChangeFramerate)
-	{
-		UE_LOG(GameplayMediaEncoder, Verbose, TEXT("framerate -> %d"), NewVideoFramerate.Load());
+		if (bChangeBitrate)
+		{
+			config.MaxBitrate = MaxVideoBitrate;
+			config.TargetBitrate = NewVideoBitrate;
+		}
 
-		VideoEncoder->UpdateFrameRate(NewVideoFramerate);
+		if (bChangeFramerate)
+		{
+			config.MaxFramerate = NewVideoFramerate;
+			NumCapturedFrames = 0;
+		}
+
+		VideoEncoder->UpdateLayerConfig(0, config);
 		bChangeFramerate = false;
-		NumCapturedFrames = 0;
+		bChangeBitrate = false;
 	}
 
 	return true;
