Incremental Linux AVEncoder compile errors
#rnx

diff --git a/Engine/Source/Runtime/AVEncoder/Private/Encoders/Amf/Amf_EncoderH264.cpp b/Engine/Source/Runtime/AVEncoder/Private/Encoders/Amf/Amf_EncoderH264.cpp
--- a/Engine/Source/Runtime/AVEncoder/Private/Encoders/Amf/Amf_EncoderH264.cpp
+++ b/Engine/Source/Runtime/AVEncoder/Private/Encoders/Amf/Amf_EncoderH264.cpp
@@ -5,6 +5,8 @@
 #include "CodecPacket.h"
 #include "AVEncoderDebug.h"
 
+#include "RHI.h"
+
 #include <stdio.h>
 
 #define MAX_GPU_INDEXES 50
@@ -539,10 +541,12 @@ namespace AVEncoder
 						UE_LOG(LogVideoEncoder, Fatal, TEXT("Amf failed to get average QP."));
 					}
 
-					if (OutBuffer->GetProperty(AMF_VIDEO_ENCODER_START_TS, &Packet.Timings.StartTs) != AMF_OK)
+					amf_int64 StartTs;
+					if (OutBuffer->GetProperty(AMF_VIDEO_ENCODER_START_TS, &StartTs) != AMF_OK)
 					{
 						UE_LOG(LogVideoEncoder, Fatal, TEXT("Amf failed to get average QP."));
 					}
+					Packet.Timings.StartTs = StartTs;
 
 					Packet.Timings.FinishTs = FTimespan::FromSeconds(FPlatformTime::Seconds());
 					Packet.Framerate = CurrentConfig.MaxFramerate;
@@ -762,8 +766,6 @@ namespace AVEncoder
 		}
 
 		EncoderInfo.CodecType = ECodecType::H264;
-		EncoderInfo.MaxWidth;
-		EncoderInfo.MaxHeight;
 
 		AMFCapsPtr EncoderCaps;
 		AmfEncoder->GetCaps(&EncoderCaps);
@@ -796,4 +798,6 @@ namespace AVEncoder
 		return bSuccess;
 	}
 
-} /* namespace AVEncoder */
\ No newline at end of file
+} /* namespace AVEncoder */
+
+#undef MIN_UPDATE_FRAMERATE_SECS
\ No newline at end of file
diff --git a/Engine/Source/Runtime/AVEncoder/Private/Encoders/Amf/Amf_EncoderH264.h b/Engine/Source/Runtime/AVEncoder/Private/Encoders/Amf/Amf_EncoderH264.h
--- a/Engine/Source/Runtime/AVEncoder/Private/Encoders/Amf/Amf_EncoderH264.h
+++ b/Engine/Source/Runtime/AVEncoder/Private/Encoders/Amf/Amf_EncoderH264.h
@@ -7,6 +7,21 @@
 #include "VideoEncoderFactory.h"
 #include "VideoEncoderInputImpl.h"
 
+THIRD_PARTY_INCLUDES_START
+
+#if PLATFORM_WINDOWS
+#include "Windows/AllowWindowsPlatformTypes.h"
+#include "Windows/PreWindowsApi.h"
+#endif
+
+#include <nvEncodeAPI.h> // To get definition of GUID
+
+#if PLATFORM_WINDOWS
+#include "Windows/PostWindowsApi.h"
+#include "Windows/HideWindowsPlatformTypes.h"
+#endif
+
+THIRD_PARTY_INCLUDES_END
 
 namespace AVEncoder
 {
diff --git a/Engine/Source/Runtime/AVEncoder/Private/Encoders/NVENC/NVENC_EncoderH264.cpp b/Engine/Source/Runtime/AVEncoder/Private/Encoders/NVENC/NVENC_EncoderH264.cpp
--- a/Engine/Source/Runtime/AVEncoder/Private/Encoders/NVENC/NVENC_EncoderH264.cpp
+++ b/Engine/Source/Runtime/AVEncoder/Private/Encoders/NVENC/NVENC_EncoderH264.cpp
@@ -1295,4 +1295,6 @@ namespace AVEncoder
         return bSuccess;
     }
 
-} /* namespace AVEncoder */
\ No newline at end of file
+} /* namespace AVEncoder */
+
+#undef MIN_UPDATE_FRAMERATE_SECS
\ No newline at end of file
