Fix Mac compilation
#rnx

diff --git a/Engine/Source/Runtime/AVEncoder/Private/VideoEncoderInput.cpp b/Engine/Source/Runtime/AVEncoder/Private/VideoEncoderInput.cpp
--- a/Engine/Source/Runtime/AVEncoder/Private/VideoEncoderInput.cpp
+++ b/Engine/Source/Runtime/AVEncoder/Private/VideoEncoderInput.cpp
@@ -7,7 +7,10 @@
 #include "AVEncoderDebug.h"
 #include "Misc/Paths.h"
 #include "VideoCommon.h"
+
+#if PLATFORM_SUPPORTS_VULKAN
 #include "VulkanRHIPrivate.h"
+#endif
 
 #if PLATFORM_WINDOWS
 #include "MicrosoftCommon.h"
