Use !PLATFORM_MAC instead of PLATFORM_SUPPORTS_VULKAN
#rnx

diff --git a/Engine/Source/Runtime/AVEncoder/Private/VideoEncoderInput.cpp b/Engine/Source/Runtime/AVEncoder/Private/VideoEncoderInput.cpp
--- a/Engine/Source/Runtime/AVEncoder/Private/VideoEncoderInput.cpp
+++ b/Engine/Source/Runtime/AVEncoder/Private/VideoEncoderInput.cpp
@@ -8,7 +8,7 @@
 #include "Misc/Paths.h"
 #include "VideoCommon.h"
 
-#if PLATFORM_SUPPORTS_VULKAN
+#if !PLATFORM_MAC
 #include "VulkanRHIPrivate.h"
 #endif
 
