Fix compilation on non-amf supported platforms
#rnx

diff --git a/Engine/Source/Runtime/AVEncoder/Private/Encoders/Amf/Amf_Common.cpp b/Engine/Source/Runtime/AVEncoder/Private/Encoders/Amf/Amf_Common.cpp
--- a/Engine/Source/Runtime/AVEncoder/Private/Encoders/Amf/Amf_Common.cpp
+++ b/Engine/Source/Runtime/AVEncoder/Private/Encoders/Amf/Amf_Common.cpp
@@ -3,6 +3,8 @@
 #include "Amf_Common.h"
 #include "AVEncoder.h"
 
+#include "RHI.h"
+
 namespace AVEncoder
 {
     FCriticalSection    FAmfCommon::ProtectSingleton;
@@ -65,10 +67,14 @@ namespace AVEncoder
 		{
 			return;
 		}
-        
+
+#ifdef AMF_DLL_NAMEA
         DllHandle = FPlatformProcess::GetDllHandle(TEXT(AMF_DLL_NAMEA));
 
         AMFInit_Fn AmfInitFn = (AMFInit_Fn)FPlatformProcess::GetDllExport(DllHandle, TEXT(AMF_INIT_FUNCTION_NAME));
+#else
+        AMFInit_Fn AmfInitFn = nullptr;
+#endif
         if (AmfInitFn == nullptr)
         {
             return;
diff --git a/Engine/Source/Runtime/AVEncoder/Private/Encoders/Amf/Amf_Common.h b/Engine/Source/Runtime/AVEncoder/Private/Encoders/Amf/Amf_Common.h
--- a/Engine/Source/Runtime/AVEncoder/Private/Encoders/Amf/Amf_Common.h
+++ b/Engine/Source/Runtime/AVEncoder/Private/Encoders/Amf/Amf_Common.h
@@ -16,9 +16,11 @@ THIRD_PARTY_INCLUDES_START
 #include "Windows/PreWindowsApi.h"
 #endif
 
+PRAGMA_DISABLE_OVERLOADED_VIRTUAL_WARNINGS
 #include "core/Factory.h"
 #include "core/Interface.h"
 #include "components/VideoEncoderVCE.h"
+PRAGMA_ENABLE_OVERLOADED_VIRTUAL_WARNINGS
 
 #if PLATFORM_WINDOWS
 #include "Windows/PostWindowsApi.h"
