Additional Mac compile fixes
#rnx

diff --git a/Engine/Source/Runtime/AVEncoder/Private/Encoders/Amf/Amf_EncoderH264.cpp b/Engine/Source/Runtime/AVEncoder/Private/Encoders/Amf/Amf_EncoderH264.cpp
--- a/Engine/Source/Runtime/AVEncoder/Private/Encoders/Amf/Amf_EncoderH264.cpp
+++ b/Engine/Source/Runtime/AVEncoder/Private/Encoders/Amf/Amf_EncoderH264.cpp
@@ -97,6 +97,7 @@ namespace AVEncoder
 			}
 			break;
 #endif
+#if !PLATFORM_MAC
 		case AVEncoder::EVideoFrameFormat::VULKAN_R8G8B8A8_UNORM:
 			EncoderDevice = Input->GetVulkanDevice();
 			if (!Amf.bIsCtxInitialized)
@@ -105,6 +106,7 @@ namespace AVEncoder
 				Amf.bIsCtxInitialized = true;
 			}
 			break;
+#endif
 		case AVEncoder::EVideoFrameFormat::Undefined:
 		default:
 			UE_LOG(LogVideoEncoder, Error, TEXT("Frame format %s is not currently supported by Amf Encoder on this platform."), *ToString(FrameFormat));
diff --git a/Engine/Source/Runtime/AVEncoder/Private/VideoEncoderInput.cpp b/Engine/Source/Runtime/AVEncoder/Private/VideoEncoderInput.cpp
--- a/Engine/Source/Runtime/AVEncoder/Private/VideoEncoderInput.cpp
+++ b/Engine/Source/Runtime/AVEncoder/Private/VideoEncoderInput.cpp
@@ -545,10 +545,12 @@ CUcontext FVideoEncoderInputImpl::GetCUDAEncoderContext() const
 
 #endif
 
+#if !PLATFORM_MAC
 VkDevice FVideoEncoderInputImpl::GetVulkanDevice() const
 {
 	return FrameInfoVulkan.VulkanDevice;
 }
+#endif
 
 // *** FVideoEncoderInputFrame ********************************************************************
 
diff --git a/Engine/Source/Runtime/AVEncoder/Private/VideoEncoderInputImpl.h b/Engine/Source/Runtime/AVEncoder/Private/VideoEncoderInputImpl.h
--- a/Engine/Source/Runtime/AVEncoder/Private/VideoEncoderInputImpl.h
+++ b/Engine/Source/Runtime/AVEncoder/Private/VideoEncoderInputImpl.h
@@ -85,7 +85,9 @@ public:
 	CUcontext GetCUDAEncoderContext() const;
 #endif
 
+#if !PLATFORM_MAC
 	VkDevice_T* GetVulkanDevice() const;
+#endif
 
 private:
 
