Clean up attempt to disable overloaded virtuals in third party code in light of CL# 16440836
#rnx

diff --git a/Engine/Source/Runtime/AVEncoder/Private/Encoders/Amf/Amf_Common.h b/Engine/Source/Runtime/AVEncoder/Private/Encoders/Amf/Amf_Common.h
--- a/Engine/Source/Runtime/AVEncoder/Private/Encoders/Amf/Amf_Common.h
+++ b/Engine/Source/Runtime/AVEncoder/Private/Encoders/Amf/Amf_Common.h
@@ -5,10 +5,6 @@
 #include <CoreMinimal.h>
 #include <HAL/Thread.h>
 
-#if PLATFORM_LINUX
-PRAGMA_DISABLE_OVERLOADED_VIRTUAL_WARNINGS
-#endif
-
 THIRD_PARTY_INCLUDES_START
 
 #if PLATFORM_WINDOWS
@@ -16,11 +12,9 @@ THIRD_PARTY_INCLUDES_START
 #include "Windows/PreWindowsApi.h"
 #endif
 
-PRAGMA_DISABLE_OVERLOADED_VIRTUAL_WARNINGS
 #include "core/Factory.h"
 #include "core/Interface.h"
 #include "components/VideoEncoderVCE.h"
-PRAGMA_ENABLE_OVERLOADED_VIRTUAL_WARNINGS
 
 #if PLATFORM_WINDOWS
 #include "Windows/PostWindowsApi.h"
@@ -29,10 +23,6 @@ PRAGMA_ENABLE_OVERLOADED_VIRTUAL_WARNINGS
 
 THIRD_PARTY_INCLUDES_END
 
-#if PLATFORM_LINUX
-PRAGMA_ENABLE_OVERLOADED_VIRTUAL_WARNINGS
-#endif
-
 #define CHECK_AMF_RET(AMF_call)\
 {\
 	AMF_RESULT Res = AMF_call;\
