Affirmatively include Vulkan for Linux and Windows rather than removing for Mac
#rnx

diff --git a/Engine/Source/Runtime/AVEncoder/Private/Encoders/Amf/Amf_EncoderH264.cpp b/Engine/Source/Runtime/AVEncoder/Private/Encoders/Amf/Amf_EncoderH264.cpp
--- a/Engine/Source/Runtime/AVEncoder/Private/Encoders/Amf/Amf_EncoderH264.cpp
+++ b/Engine/Source/Runtime/AVEncoder/Private/Encoders/Amf/Amf_EncoderH264.cpp
@@ -97,7 +97,7 @@ namespace AVEncoder
 			}
 			break;
 #endif
-#if !PLATFORM_MAC
+#if PLATFORM_WINDOWS || PLATFORM_LINUX
 		case AVEncoder::EVideoFrameFormat::VULKAN_R8G8B8A8_UNORM:
 			EncoderDevice = Input->GetVulkanDevice();
 			if (!Amf.bIsCtxInitialized)
diff --git a/Engine/Source/Runtime/AVEncoder/Private/VideoEncoderInput.cpp b/Engine/Source/Runtime/AVEncoder/Private/VideoEncoderInput.cpp
--- a/Engine/Source/Runtime/AVEncoder/Private/VideoEncoderInput.cpp
+++ b/Engine/Source/Runtime/AVEncoder/Private/VideoEncoderInput.cpp
@@ -8,7 +8,7 @@
 #include "Misc/Paths.h"
 #include "VideoCommon.h"
 
-#if !PLATFORM_MAC
+#if PLATFORM_WINDOWS || PLATFORM_LINUX
 #include "VulkanRHIPrivate.h"
 #endif
 
@@ -545,7 +545,7 @@ CUcontext FVideoEncoderInputImpl::GetCUDAEncoderContext() const
 
 #endif
 
-#if !PLATFORM_MAC
+#if PLATFORM_WINDOWS || PLATFORM_LINUX
 VkDevice FVideoEncoderInputImpl::GetVulkanDevice() const
 {
 	return FrameInfoVulkan.VulkanDevice;
@@ -802,6 +802,7 @@ void FVideoEncoderInputFrame::SetTexture(CUarray InTexture, FReleaseCUDATextureC
 #endif
 
 
+#if PLATFORM_WINDOWS || PLATFORM_LINUX
 void FVideoEncoderInputFrame::SetTexture(VkImage InTexture, FReleaseVulkanTextureCallback InOnReleaseTexture)
 {
 	if (Format == EVideoFrameFormat::VULKAN_R8G8B8A8_UNORM)
@@ -814,6 +815,7 @@ void FVideoEncoderInputFrame::SetTexture(VkImage InTexture, FReleaseVulkanTextur
 		}
 	}
 }
+#endif
 
 // *** FVideoEncoderInputFrameImpl ****************************************************************
 
diff --git a/Engine/Source/Runtime/AVEncoder/Private/VideoEncoderInputImpl.h b/Engine/Source/Runtime/AVEncoder/Private/VideoEncoderInputImpl.h
--- a/Engine/Source/Runtime/AVEncoder/Private/VideoEncoderInputImpl.h
+++ b/Engine/Source/Runtime/AVEncoder/Private/VideoEncoderInputImpl.h
@@ -85,7 +85,7 @@ public:
 	CUcontext GetCUDAEncoderContext() const;
 #endif
 
-#if !PLATFORM_MAC
+#if PLATFORM_WINDOWS || PLATFORM_LINUX
 	VkDevice_T* GetVulkanDevice() const;
 #endif
 
@@ -130,10 +130,12 @@ private:
 	}								FrameInfoCUDA;
 #endif
 
+#if PLATFORM_WINDOWS || PLATFORM_LINUX
 	struct FFrameInfoVulkan
 	{
 		VkDevice_T*					VulkanDevice;
 	}								FrameInfoVulkan;
+#endif
 
 	mutable FCriticalSection				ProtectFrames;
 	TQueue<FVideoEncoderInputFrameImpl*>	AvailableFrames;
diff --git a/Engine/Source/Runtime/AVEncoder/Public/VideoEncoderInput.h b/Engine/Source/Runtime/AVEncoder/Public/VideoEncoderInput.h
--- a/Engine/Source/Runtime/AVEncoder/Public/VideoEncoderInput.h
+++ b/Engine/Source/Runtime/AVEncoder/Public/VideoEncoderInput.h
@@ -212,7 +212,9 @@ namespace AVEncoder
 		// the callback type used to create a registered encoder
 		using FReleaseVulkanTextureCallback = TFunction<void(VkImage_T*)>;
 
-		void SetTexture(VkImage_T* InTexture, FReleaseVulkanTextureCallback InOnReleaseTexture);
+#if PLATFORM_WINDOWS || PLATFORM_LINUX
+	void SetTexture(VkImage_T* InTexture, FReleaseVulkanTextureCallback InOnReleaseTexture);
+#endif
 
 	protected:
 		FVideoEncoderInputFrame();
