Only include VulkanRHI as a dependency on Win32, Win64, or Linux
#rnx

diff --git a/Engine/Source/Runtime/AVEncoder/AVEncoder.Build.cs b/Engine/Source/Runtime/AVEncoder/AVEncoder.Build.cs
--- a/Engine/Source/Runtime/AVEncoder/AVEncoder.Build.cs
+++ b/Engine/Source/Runtime/AVEncoder/AVEncoder.Build.cs
@@ -40,10 +40,11 @@ public class AVEncoder : ModuleRules
 		DynamicallyLoadedModuleNames.AddRange(new string[] {
 			// ... add any modules that your module loads dynamically here ...
 		});
-		
-		PrivateIncludePathModuleNames.AddRange(new string[] {
-			"VulkanRHI"
-		});
+
+		if (Target.Platform == UnrealTargetPlatform.Win64 || Target.Platform == UnrealTargetPlatform.Win32 || Target.Platform == UnrealTargetPlatform.Linux)
+		{
+			PrivateIncludePathModuleNames.Add("VulkanRHI");
+		}
 
 		string EngineSourceDirectory = Path.GetFullPath(Target.RelativeEnginePath);
 		PrivateIncludePaths.Add(Path.Combine(EngineSourceDirectory, "Source/Runtime/VulkanRHI/Private"));
