Remove win32

diff --git a/Engine/Source/Runtime/AVEncoder/AVEncoder.Build.cs b/Engine/Source/Runtime/AVEncoder/AVEncoder.Build.cs
--- a/Engine/Source/Runtime/AVEncoder/AVEncoder.Build.cs
+++ b/Engine/Source/Runtime/AVEncoder/AVEncoder.Build.cs
@@ -41,7 +41,7 @@ public class AVEncoder : ModuleRules
 			// ... add any modules that your module loads dynamically here ...
 		});
 
-		if (Target.Platform == UnrealTargetPlatform.Win64 || Target.Platform == UnrealTargetPlatform.Win32 || Target.Platform == UnrealTargetPlatform.Linux)
+		if (Target.Platform == UnrealTargetPlatform.Win64 || Target.Platform == UnrealTargetPlatform.Linux)
 		{
 			PrivateIncludePathModuleNames.Add("VulkanRHI");
 		}
