Fix windows clang compile error ISO C++11 does not allow conversion from string literal to 'LPTSTR' (aka 'wchar_t *') [-Werror,-Wwritable-strings]
#rnx

diff --git a/Engine/Source/Runtime/AVEncoder/Private/Encoders/Amf/Amf_EncoderH264.cpp b/Engine/Source/Runtime/AVEncoder/Private/Encoders/Amf/Amf_EncoderH264.cpp
--- a/Engine/Source/Runtime/AVEncoder/Private/Encoders/Amf/Amf_EncoderH264.cpp
+++ b/Engine/Source/Runtime/AVEncoder/Private/Encoders/Amf/Amf_EncoderH264.cpp
@@ -304,7 +304,8 @@ namespace AVEncoder
 			}
 			else if (WaitResult == WAIT_FAILED)
 			{
-				WindowsError(TEXT("WaitForMultipleObjects"));
+				const FString WaitError(TEXT("WaitForMultipleObjects"));
+				WindowsError((LPTSTR)*WaitError);
 			}
 		}
 #else
