Fix Win64 Release build cs reference to platform removed for release builds

diff --git a/Engine/Source/Runtime/AVEncoder/AVEncoder.Build.cs b/Engine/Source/Runtime/AVEncoder/AVEncoder.Build.cs
--- a/Engine/Source/Runtime/AVEncoder/AVEncoder.Build.cs
+++ b/Engine/Source/Runtime/AVEncoder/AVEncoder.Build.cs
@@ -61,18 +61,15 @@ public class AVEncoder : ModuleRules
 				"D3D12RHI"
 			});
 
-			if (Target.Platform != UnrealTargetPlatform.XboxOneGDK)
-			{
-				PublicSystemLibraries.AddRange(new string[] {
-					"DXGI.lib",
-					"d3d11.lib",
-					"d3d12.lib"
-				});
-				
-				PublicDelayLoadDLLs.Add("mfplat.dll");
-				PublicDelayLoadDLLs.Add("mfuuid.dll");
-				PublicDelayLoadDLLs.Add("Mfreadwrite.dll");
-			}
+			PublicSystemLibraries.AddRange(new string[] {
+				"DXGI.lib",
+				"d3d11.lib",
+				"d3d12.lib"
+			});
+			
+			PublicDelayLoadDLLs.Add("mfplat.dll");
+			PublicDelayLoadDLLs.Add("mfuuid.dll");
+			PublicDelayLoadDLLs.Add("Mfreadwrite.dll");
 		}
 	}
 }
