Decouple render FPS and encode FPS so that low FPS application can still stream at 60 FPS.
This resolves bitrate issues in WebRTC caused by inconsistent frame rates.

These changes are intended for 5.0 and licencees who are using source builds of the engine to get the updates before 5.0 is released.

diff --git a/Engine/Plugins/Media/HardwareEncoders/Source/EncoderAMF/Private/Amf_EncoderH264.cpp b/Engine/Plugins/Media/HardwareEncoders/Source/EncoderAMF/Private/Amf_EncoderH264.cpp
--- a/Engine/Plugins/Media/HardwareEncoders/Source/EncoderAMF/Private/Amf_EncoderH264.cpp
+++ b/Engine/Plugins/Media/HardwareEncoders/Source/EncoderAMF/Private/Amf_EncoderH264.cpp
@@ -26,10 +26,10 @@ namespace
 	{
 		switch (mode)
 		{
-		case AVEncoder::FVideoEncoder::RateControlMode::CONSTQP: return AMF_VIDEO_ENCODER_RATE_CONTROL_METHOD_CONSTANT_QP;
-		case AVEncoder::FVideoEncoder::RateControlMode::VBR: return AMF_VIDEO_ENCODER_RATE_CONTROL_METHOD_LATENCY_CONSTRAINED_VBR;
-		default:
-		case AVEncoder::FVideoEncoder::RateControlMode::CBR: return AMF_VIDEO_ENCODER_RATE_CONTROL_METHOD_CBR;
+			case AVEncoder::FVideoEncoder::RateControlMode::CONSTQP: return AMF_VIDEO_ENCODER_RATE_CONTROL_METHOD_CONSTANT_QP;
+			case AVEncoder::FVideoEncoder::RateControlMode::VBR: return AMF_VIDEO_ENCODER_RATE_CONTROL_METHOD_LATENCY_CONSTRAINED_VBR;
+			default:
+			case AVEncoder::FVideoEncoder::RateControlMode::CBR: return AMF_VIDEO_ENCODER_RATE_CONTROL_METHOD_CBR;
 		}
 	}
 
@@ -37,13 +37,13 @@ namespace
 	{
 		switch (profile)
 		{
-		case AVEncoder::FVideoEncoder::H264Profile::CONSTRAINED_BASELINE: return AMF_VIDEO_ENCODER_PROFILE_CONSTRAINED_BASELINE;
-		case AVEncoder::FVideoEncoder::H264Profile::BASELINE: return AMF_VIDEO_ENCODER_PROFILE_BASELINE;
-		case AVEncoder::FVideoEncoder::H264Profile::MAIN: return AMF_VIDEO_ENCODER_PROFILE_MAIN;
-		case AVEncoder::FVideoEncoder::H264Profile::CONSTRAINED_HIGH: return AMF_VIDEO_ENCODER_PROFILE_CONSTRAINED_HIGH;
-		case AVEncoder::FVideoEncoder::H264Profile::HIGH: return AMF_VIDEO_ENCODER_PROFILE_HIGH;
-		default:
-		case AVEncoder::FVideoEncoder::H264Profile::AUTO: return AMF_VIDEO_ENCODER_PROFILE_BASELINE;
+			case AVEncoder::FVideoEncoder::H264Profile::CONSTRAINED_BASELINE: return AMF_VIDEO_ENCODER_PROFILE_CONSTRAINED_BASELINE;
+			case AVEncoder::FVideoEncoder::H264Profile::BASELINE: return AMF_VIDEO_ENCODER_PROFILE_BASELINE;
+			case AVEncoder::FVideoEncoder::H264Profile::MAIN: return AMF_VIDEO_ENCODER_PROFILE_MAIN;
+			case AVEncoder::FVideoEncoder::H264Profile::CONSTRAINED_HIGH: return AMF_VIDEO_ENCODER_PROFILE_CONSTRAINED_HIGH;
+			case AVEncoder::FVideoEncoder::H264Profile::HIGH: return AMF_VIDEO_ENCODER_PROFILE_HIGH;
+			default:
+			case AVEncoder::FVideoEncoder::H264Profile::AUTO: return AMF_VIDEO_ENCODER_PROFILE_BASELINE;
 		}
 	}
 }
@@ -88,20 +88,10 @@ namespace AVEncoder
 
 	void AMFParseCommandLineFlags()
 	{
-		// CVar changes can only be triggered from the game thread
-		{
-			// We block on our current thread until these CVars are set on the game thread, the settings of these CVars needs to happen before we can proceed.
-            FScopedEvent ThreadBlocker;
-
-			AsyncTask(ENamedThreads::GameThread, [&ThreadBlocker]()
-			{ 
-				AMFCommandLineParseValue(TEXT("-AMFKeyframeInterval="), CVarAMFKeyframeInterval);
-
-				// Unblocks the thread
-                ThreadBlocker.Trigger();
-			});
-            // Blocks current thread until game thread calls trigger (indicating it is done)
-        }
+		AsyncTask(ENamedThreads::GameThread, []()
+		{ 
+			AMFCommandLineParseValue(TEXT("-AMFKeyframeInterval="), CVarAMFKeyframeInterval);
+		});
 	}
 
 	static bool GetEncoderInfo(FAmfCommon& Amf, FVideoEncoderInfo& EncoderInfo);
@@ -135,6 +125,7 @@ namespace AVEncoder
 	FVideoEncoderAmf_H264::FVideoEncoderAmf_H264()
 		: Amf(FAmfCommon::Setup())
 	{
+		AMFParseCommandLineFlags();
 	}
 
 	FVideoEncoderAmf_H264::~FVideoEncoderAmf_H264()
@@ -203,8 +194,6 @@ namespace AVEncoder
 			}
 		}
 
-		AMFParseCommandLineFlags();
-
 		FLayerConfig MutableConfig = config;
 		if (MutableConfig.MaxFramerate == 0)
 		{
@@ -233,16 +222,14 @@ namespace AVEncoder
 
 	void FVideoEncoderAmf_H264::Encode(FVideoEncoderInputFrame const* frame, FEncodeOptions const& options)
 	{
-		// todo: reconfigure encoder
-		auto const amfFrame = static_cast<const FVideoEncoderInputFrameImpl*>(frame);
+		const FVideoEncoderInputFrameImpl* amfFrame = static_cast<const FVideoEncoderInputFrameImpl*>(frame);
 		for (auto& layer : Layers)
 		{
 			FAMFLayer* amfLayer = static_cast<FAMFLayer*>(layer);
 			AMF_RESULT Res = amfLayer->Encode(amfFrame, options);
-
-			if (Res == AMF_OK)
+			if(Res != AMF_OK)
 			{
-				amfLayer->FramesPending->Trigger();
+				UE_LOG(LogEncoderAMF, Error, TEXT("AMF failed to encode frame."));
 			}
 		}
 	}
@@ -266,7 +253,7 @@ namespace AVEncoder
 		}
 
 		Layers.Reset();
-		}
+	}
 
 	// --- Amf_EncoderH264::FLayer ------------------------------------------------------------
 	FVideoEncoderAmf_H264::FAMFLayer::FAMFLayer(uint32 layerIdx, FLayerConfig const& config, FVideoEncoderAmf_H264& encoder)
@@ -293,11 +280,6 @@ namespace AVEncoder
 			Amf.CreateEncoder(AmfEncoder);
 		}
 
-		if (ProcessFrameThread == nullptr)
-		{
-			ProcessFrameThread = MakeUnique<FThread>(TEXT("AmfFrameProcessingThread"), [this]() { ProcessFrameThreadFunc(); });
-		}
-
 		return AmfEncoder != NULL;
 	}
 
@@ -423,7 +405,7 @@ namespace AVEncoder
 			MaybeReconfigure();
 
 			Buffer->Surface->SetPts(frame->GetTimestampRTP());
-			amf_int64 Start_ts = FTimespan::FromSeconds(FPlatformTime::Seconds()).GetTicks();
+			amf_int64 Start_ts = FPlatformTime::Cycles64();
 			Buffer->Surface->SetProperty(AMF_VIDEO_ENCODER_START_TS, Start_ts);
 			Buffer->Surface->SetProperty(AMF_BUFFER_INPUT_FRAME, uintptr_t(frame));
 
@@ -464,7 +446,8 @@ namespace AVEncoder
 				}
 				else
 				{
-					PendingFrames.Increment();
+					// Note from Luke: Testing has shown this is okay to leave blocking on the calling thread.
+					ProcessFrameBlocking();
 				}
 			}
 		}
@@ -482,14 +465,6 @@ namespace AVEncoder
 		Flush();
 		CreatedSurfaces.Empty();
 
-		if (ProcessFrameThread != nullptr)
-		{
-			bShouldRunProcessingThread = false;
-			FramesPending->Trigger();
-			ProcessFrameThread->Join();
-			ProcessFrameThread = nullptr;
-		}
-
 		if (AmfEncoder != NULL)
 		{
 			AmfEncoder->Terminate();
@@ -497,82 +472,73 @@ namespace AVEncoder
 		}
 	}
 
-	void FVideoEncoderAmf_H264::FAMFLayer::ProcessFrameThreadFunc()
+	void FVideoEncoderAmf_H264::FAMFLayer::ProcessFrameBlocking()
 	{
-		bool bHasProcessedFrame = false;
-		while (bShouldRunProcessingThread)
+		checkf(!bIsProcessingFrame, TEXT("There is already a frame being processing in the AMF encoder. Only one thing should call encode at a time."));
+
+		bIsProcessingFrame = true;
+
+		while(bIsProcessingFrame)
 		{
-			if (bWaitingForFrames) {
-				FramesPending->Wait();
+			amf::AMFDataPtr data;
+			AMF_RESULT Result = AmfEncoder->QueryOutput(&data);
+
+			if (Result != AMF_OK)
+			{
+				// No frame ready yet in output.
+				continue;
 			}
 
-			if (PendingFrames.GetValue() > 0)
+			AMFBufferPtr OutBuffer(data);
+
+			// Create packet with buffer contents
+			FCodecPacketImpl Packet;
+
+			Packet.Data = static_cast<const uint8*>(OutBuffer->GetNative());
+			Packet.DataSize = OutBuffer->GetSize();
+			uint32 PictureType = AMF_VIDEO_ENCODER_PICTURE_TYPE_NONE;
+			if (OutBuffer->GetProperty(AMF_VIDEO_ENCODER_OUTPUT_DATA_TYPE, &PictureType) != AMF_OK)
 			{
-				PendingFrames.Decrement();
-
-					amf::AMFDataPtr data;
-				AMF_RESULT Result = AmfEncoder->QueryOutput(&data);
-
-				if (Result == AMF_OK)
-					{
-						AMFBufferPtr OutBuffer(data);
-
-						// create packet with buffer contents
-						FCodecPacketImpl Packet;
-
-						Packet.Data = static_cast<const uint8*>(OutBuffer->GetNative());
-						Packet.DataSize = OutBuffer->GetSize();
-						uint32 PictureType = AMF_VIDEO_ENCODER_PICTURE_TYPE_NONE;
-						if (OutBuffer->GetProperty(AMF_VIDEO_ENCODER_OUTPUT_DATA_TYPE, &PictureType) != AMF_OK)
-						{
-							UE_LOG(LogEncoderAMF, Error, TEXT("Amf failed to get picture type."));
-						}
-						else if (PictureType == AMF_VIDEO_ENCODER_OUTPUT_DATA_TYPE_IDR)
-						{
-							UE_LOG(LogEncoderAMF, Verbose, TEXT("Generated IDR Frame"));
-							Packet.IsKeyFrame = true;
-						}
-
-						if (FString(GDynamicRHI->GetName()) != FString("Vulkan")) // Amf with Vulkan doesn't currently support statistics
-						{
-							if (OutBuffer->GetProperty(AMF_VIDEO_ENCODER_STATISTIC_FRAME_QP, &Packet.VideoQP) != AMF_OK)
-							{
-								UE_LOG(LogEncoderAMF, Error, TEXT("Amf failed to get frame QP."));
-							}
-						}
-
-						amf_int64 StartTs;
-						if (OutBuffer->GetProperty(AMF_VIDEO_ENCODER_START_TS, &StartTs) != AMF_OK)
-						{
-							UE_LOG(LogEncoderAMF, Error, TEXT("Amf failed to get encode start time."));
-						}
-						Packet.Timings.StartTs = FTimespan(StartTs);
-
-						Packet.Timings.FinishTs = FTimespan::FromSeconds(FPlatformTime::Seconds());
-					Packet.Framerate = GetConfig().MaxFramerate;
-
-						FVideoEncoderInputFrameImpl* SourceFrame;
-						if (OutBuffer->GetProperty(AMF_BUFFER_INPUT_FRAME, (intptr_t*)&SourceFrame) != AMF_OK)
-						{
-							UE_LOG(LogEncoderAMF, Fatal, TEXT("Amf failed to get buffer input frame."));
-						}
-
-					if (Encoder.OnEncodedPacket)
-						{
-						Encoder.OnEncodedPacket(LayerIndex, SourceFrame, Packet);
-						}
-
-						bHasProcessedFrame = true;
-					}
-
-			if (!bHasProcessedFrame)
+				UE_LOG(LogEncoderAMF, Error, TEXT("Amf failed to get picture type."));
+			}
+			else if (PictureType == AMF_VIDEO_ENCODER_OUTPUT_DATA_TYPE_IDR)
 			{
-				bWaitingForFrames = true;
+				UE_LOG(LogEncoderAMF, Verbose, TEXT("Generated IDR Frame"));
+				Packet.IsKeyFrame = true;
 			}
 
-			bHasProcessedFrame = false;
+			if (FString(GDynamicRHI->GetName()) != FString("Vulkan")) // Amf with Vulkan doesn't currently support statistics
+			{
+				if (OutBuffer->GetProperty(AMF_VIDEO_ENCODER_STATISTIC_FRAME_QP, &Packet.VideoQP) != AMF_OK)
+				{
+					UE_LOG(LogEncoderAMF, Error, TEXT("Amf failed to get frame QP."));
+				}
+			}
+
+			amf_int64 StartTs;
+			if (OutBuffer->GetProperty(AMF_VIDEO_ENCODER_START_TS, &StartTs) != AMF_OK)
+			{
+				UE_LOG(LogEncoderAMF, Error, TEXT("Amf failed to get encode start time."));
+			}
+
+			Packet.Timings.StartTs = FTimespan::FromSeconds(FPlatformTime::ToSeconds64(StartTs));
+			Packet.Timings.FinishTs = FTimespan::FromSeconds(FPlatformTime::ToSeconds64(FPlatformTime::Cycles64()));
+			Packet.Framerate = GetConfig().MaxFramerate;
+
+			FVideoEncoderInputFrameImpl* SourceFrame;
+			if (OutBuffer->GetProperty(AMF_BUFFER_INPUT_FRAME, (intptr_t*)&SourceFrame) != AMF_OK)
+			{
+				UE_LOG(LogEncoderAMF, Fatal, TEXT("Amf failed to get buffer input frame."));
+			}
+
+			if (Encoder.OnEncodedPacket)
+			{
+				Encoder.OnEncodedPacket(LayerIndex, SourceFrame, Packet);
+			}
+
+			bIsProcessingFrame = false;
 		}
-	}
+		
 	}
 
 	template<class T>
diff --git a/Engine/Plugins/Media/HardwareEncoders/Source/EncoderAMF/Private/Amf_EncoderH264.h b/Engine/Plugins/Media/HardwareEncoders/Source/EncoderAMF/Private/Amf_EncoderH264.h
--- a/Engine/Plugins/Media/HardwareEncoders/Source/EncoderAMF/Private/Amf_EncoderH264.h
+++ b/Engine/Plugins/Media/HardwareEncoders/Source/EncoderAMF/Private/Amf_EncoderH264.h
@@ -40,6 +40,14 @@ namespace AVEncoder
 		class FAMFLayer : public FLayer
 		{
 		public:
+
+			class FInputOutput
+			{
+				public:
+					void* TextureToCompress;
+					AMFSurfacePtr Surface;
+			};
+
 			FAMFLayer(uint32 layerIdx, FLayerConfig const& config, FVideoEncoderAmf_H264& encoder);
 			~FAMFLayer();
 
@@ -59,39 +67,23 @@ namespace AVEncoder
 			void Shutdown();
 			void UpdateBitrate(uint32 InMaxBitRate, uint32 InTargetBitRate);
 			void UpdateResolution(uint32 InMaxBitRate, uint32 InTargetBitRate);
+			void MaybeReconfigure();
+			void ProcessFrameBlocking();
+			TSharedPtr<FInputOutput> GetOrCreateSurface(const FVideoEncoderInputFrameImpl* InFrame);
+			bool CreateSurface(TSharedPtr<FInputOutput>& OutBuffer, const FVideoEncoderInputFrameImpl* SourceFrame, void* TextureToCompress);
 
 			FVideoEncoderAmf_H264& Encoder;
 			FAmfCommon& Amf;
 			uint32 LayerIndex;
 			AMFComponentPtr AmfEncoder = NULL;
-			FThreadSafeCounter PendingFrames;
 			FDateTime LastKeyFrameTime = 0;
 			bool bForceNextKeyframe = false;
-
-			class FInputOutput
-			{
-			public:
-				void* TextureToCompress;
-				AMFSurfacePtr Surface;
-			};
-
-			void MaybeReconfigure();
-
-			TSharedPtr<FInputOutput> GetOrCreateSurface(const FVideoEncoderInputFrameImpl* InFrame);
-			bool CreateSurface(TSharedPtr<FInputOutput>& OutBuffer, const FVideoEncoderInputFrameImpl* SourceFrame, void* TextureToCompress);
-
 			TArray<TSharedPtr<FInputOutput>> CreatedSurfaces;
 			FThreadSafeBool bUpdateConfig = false;
-
 			uint32 CurrentWidth;
 			uint32 CurrentHeight;
 			uint32 CurrentFrameRate;
-
-			void ProcessFrameThreadFunc();
-			TUniquePtr<FThread> ProcessFrameThread;
-			FThreadSafeBool bShouldRunProcessingThread = true;
-			FThreadSafeBool bWaitingForFrames = true;
-			FEventRef FramesPending;
+			FThreadSafeBool bIsProcessingFrame = false;
 		};
 
 		FAmfCommon& Amf;
diff --git a/Engine/Plugins/Media/HardwareEncoders/Source/EncoderNVENC/Private/NVENCStats.h b/Engine/Plugins/Media/HardwareEncoders/Source/EncoderNVENC/Private/NVENCStats.h
new file mode 100644
--- /dev/null
+++ b/Engine/Plugins/Media/HardwareEncoders/Source/EncoderNVENC/Private/NVENCStats.h
@@ -0,0 +1,148 @@
+// Copyright Epic Games, Inc. All Rights Reserved.
+
+#pragma once
+
+#include "NVENC_Common.h"
+#include "Tickable.h"
+#include "Stats/Stats.h"
+#include "Engine/Engine.h"
+
+// Stats about NVENC that can displayed either in the in-application HUD or in the log.
+class FNVENCStats : FTickableGameObject
+{
+	public:
+
+        FNVENCStats()
+            : bOutputToLog(false)
+            , bOutputToScreen(false)
+        {
+
+        }
+
+		static FNVENCStats& Get()
+        {
+            static FNVENCStats Stats;
+	        return Stats;
+        }
+
+        void Tick(float DeltaTime)
+        {
+            if (!GEngine)
+            {
+                return;
+            }
+
+            if(!bOutputToScreen && !bOutputToLog)
+            {
+                return;
+            }
+
+            int Id = 99;
+            EmitStat(++Id, FString::Printf(TEXT("MaybeReconfigure(): %.5f ms"), ReconfigureMs.Get()));
+            EmitStat(++Id, FString::Printf(TEXT("UnlockOutputBuffer(): %.5f ms"), UnlockOutputBufferMs.Get()));
+            EmitStat(++Id, FString::Printf(TEXT("LockOutputBuffer(): %.5f ms"), LockOutputBufferMs.Get()));
+            EmitStat(++Id, FString::Printf(TEXT("nvEncEncodePicture(): %.5f ms"), nvEncEncodePictureMs.Get()));
+            EmitStat(++Id, FString::Printf(TEXT("ProcessFramesFunc(): %.5f ms"), ProcessFramesFuncMs.Get()));
+            EmitStat(++Id, FString::Printf(TEXT("QueueEncode(): %.5f ms"), QueueEncodeMs.Get()));
+            EmitStat(++Id, FString::Printf(TEXT("Total encoder latency: %.5f ms"), TotalEncoderLatencyMs.Get()));
+            EmitStat(++Id, "------------ NVENC Stats ------------");
+        }
+
+        FORCEINLINE TStatId GetStatId() const { RETURN_QUICK_DECLARE_CYCLE_STAT(FNVENCStats, STATGROUP_Tickables); }
+
+        void SetOutputToScreen(bool bInOutputToScreen)
+        {
+            bOutputToScreen = bInOutputToScreen;
+            GAreScreenMessagesEnabled = bInOutputToScreen;
+        }
+
+        void SetOutputToLog(bool bInOutputToLog)
+        {
+            bOutputToLog = bInOutputToLog;
+        }
+
+		void SetProcessFramesFuncLatency(double InProcessFramesFuncMs)
+        {
+            this->ProcessFramesFuncMs.Update(InProcessFramesFuncMs);
+        }
+
+        void SetTotalEncoderLatency(double InTotalEncoderLatencyMs)
+        {
+            this->TotalEncoderLatencyMs.Update(InTotalEncoderLatencyMs);
+        }
+
+        void SetnvEncEncodePictureLatency(double InnvEncEncodePictureMs)
+        {
+            this->nvEncEncodePictureMs.Update(InnvEncEncodePictureMs);
+        }
+
+        void SetReconfigureLatency(double InReconfigureMs)
+        {
+            this->ReconfigureMs.Update(InReconfigureMs);
+        }
+
+        void SetQueueEncodeLatency(double InQueueEncodeMs)
+        {
+            this->QueueEncodeMs.Update(InQueueEncodeMs);
+        }
+
+        void SetLockOutputBufferLatency(double InLockOutputBufferMs)
+        {
+            this->LockOutputBufferMs.Update(InLockOutputBufferMs);
+        }
+
+        void SetUnlockOutputBufferLatency(double InUnlockOutputBufferMs)
+        {
+            this->UnlockOutputBufferMs.Update(InUnlockOutputBufferMs);
+        }
+
+	private:
+		void EmitStat(int UniqueId, FString StringToEmit)
+        {
+            GEngine->AddOnScreenDebugMessage(UniqueId, 0, FColor::Green, StringToEmit, false /* newer on top */);
+
+            if(this->bOutputToLog)
+            {
+                UE_LOG(LogEncoderNVENC, Log, TEXT("%s"), *StringToEmit);
+            }
+        }
+
+    private:
+
+        template<uint32 SmoothingPeriod>
+        class FSmoothedValue
+        {
+        public:
+            double Get() const
+            {
+                return Value;
+            }
+
+            void Update(double NewValue)
+            {
+                Value = Value * (SmoothingPeriod - 1) / SmoothingPeriod + NewValue / SmoothingPeriod;
+            }
+
+            void Reset()
+            {
+                Value = 0;
+            }
+
+        private:
+            TAtomic<double> Value{ 0 };
+        };
+
+	private:
+
+        static constexpr uint32 SmoothingPeriod = 3 * 60; 
+        bool bOutputToLog;
+        bool bOutputToScreen;
+
+		FSmoothedValue<SmoothingPeriod> ProcessFramesFuncMs;
+        FSmoothedValue<SmoothingPeriod> TotalEncoderLatencyMs;
+        FSmoothedValue<SmoothingPeriod> nvEncEncodePictureMs;
+        FSmoothedValue<SmoothingPeriod> QueueEncodeMs;
+        FSmoothedValue<SmoothingPeriod> ReconfigureMs;
+        FSmoothedValue<SmoothingPeriod> LockOutputBufferMs;
+        FSmoothedValue<SmoothingPeriod> UnlockOutputBufferMs;
+};
\ No newline at end of file
diff --git a/Engine/Plugins/Media/HardwareEncoders/Source/EncoderNVENC/Private/NVENC_EncoderH264.cpp b/Engine/Plugins/Media/HardwareEncoders/Source/EncoderNVENC/Private/NVENC_EncoderH264.cpp
--- a/Engine/Plugins/Media/HardwareEncoders/Source/EncoderNVENC/Private/NVENC_EncoderH264.cpp
+++ b/Engine/Plugins/Media/HardwareEncoders/Source/EncoderNVENC/Private/NVENC_EncoderH264.cpp
@@ -14,6 +14,7 @@
 #include "Misc/ScopedEvent.h"
 #include "Async/Async.h"
 #include <stdio.h>
+#include "NVENCStats.h"
 
 #define MAX_GPU_INDEXES 50
 #define DEFAULT_BITRATE 1000000u
@@ -81,25 +82,31 @@ namespace AVEncoder
 {
 	// Console variables for NVENC
 
-	TAutoConsoleVariable<int32>  CVarIntraRefreshPeriodFrames(
+	TAutoConsoleVariable<int32>  CVarNVENCIntraRefreshPeriodFrames(
 	TEXT("NVENC.IntraRefreshPeriodFrames"),
 	0,
 	TEXT("The total number of frames between each intra refresh. Smallers values will cause intra refresh more often. Default: 0. Values <= 0 will disable intra refresh."),
 	ECVF_Default);
 
-	TAutoConsoleVariable<int32>  CVarIntraRefreshCountFrames(
+	TAutoConsoleVariable<int32>  CVarNVENCIntraRefreshCountFrames(
 	TEXT("NVENC.IntraRefreshCountFrames"),
 	0,
 	TEXT("The total number of frames within the intra refresh period that should be used as 'intra refresh' frames. Smaller values make stream recovery quicker at the cost of more bandwidth usage. Default: 0."),
 	ECVF_Default);
 
-	TAutoConsoleVariable<bool>  CVarKeyframeQPUseLastQP(
+	TAutoConsoleVariable<bool>  CVarNVENCKeyframeQPUseLastQP(
 	TEXT("NVENC.KeyframeQPUseLastQP"),
 	true,
 	TEXT("If true QP of keyframes is no worse than the last frame transmitted (may cost latency), if false, it may be keyframe QP may be worse if network conditions require it (lower latency/worse quality). Default: true."),
 	ECVF_Default);
 
-	TAutoConsoleVariable<int32>  CVarKeyframeInterval(
+	TAutoConsoleVariable<bool>  CVarNVENCEnableStats(
+	TEXT("NVENC.EnableStats"),
+	false,
+	TEXT("Whether to enable NVENC stats or not. Default: false."),
+	ECVF_Default);
+
+	TAutoConsoleVariable<int32>  CVarNVENCKeyframeInterval(
 	TEXT("NVENC.KeyframeInterval"),
 	300,
 	TEXT("Every N frames an IDR frame is sent. Default: 300. Note: A value <= 0 will disable sending of IDR frames on an interval."),
@@ -137,22 +144,24 @@ namespace AVEncoder
 	void NVENCParseCommandLineFlags()
 	{
 		// CVar changes can only be triggered from the game thread
-		{
-			// We block on our current thread until these CVars are set on the game thread, the settings of these CVars needs to happen before we can proceed.
-            FScopedEvent ThreadBlocker;
-
-			AsyncTask(ENamedThreads::GameThread, [&ThreadBlocker]()
-			{ 
-				NVENCCommandLineParseValue(TEXT("-NVENCIntraRefreshPeriodFrames="), CVarIntraRefreshPeriodFrames);
-				NVENCCommandLineParseValue(TEXT("-NVENCIntraRefreshCountFrames="), CVarIntraRefreshCountFrames);
-				NVENCCommandLineParseOption(TEXT("-NVENCKeyFrameQPUseLastQP="), CVarKeyframeQPUseLastQP);
-				NVENCCommandLineParseValue(TEXT("-NVENCKeyframeInterval="), CVarKeyframeInterval);
-
-				// Unblocks the thread
-                ThreadBlocker.Trigger();
-			});
-            // Blocks current thread until game thread calls trigger (indicating it is done)
-        }
+
+		AsyncTask(ENamedThreads::GameThread, []()
+		{ 
+			// Values
+			NVENCCommandLineParseValue(TEXT("-NVENCIntraRefreshPeriodFrames="), CVarNVENCIntraRefreshPeriodFrames);
+			NVENCCommandLineParseValue(TEXT("-NVENCIntraRefreshCountFrames="), CVarNVENCIntraRefreshCountFrames);
+			NVENCCommandLineParseValue(TEXT("-NVENCKeyframeInterval="), CVarNVENCKeyframeInterval);
+
+			// Options
+			NVENCCommandLineParseOption(TEXT("-NVENCKeyFrameQPUseLastQP"), CVarNVENCKeyframeQPUseLastQP);
+			NVENCCommandLineParseOption(TEXT("-NVENCEnableStats"), CVarNVENCEnableStats);
+
+			if(CVarNVENCEnableStats.GetValueOnAnyThread())
+			{
+				FNVENCStats::Get().SetOutputToScreen(true);
+			}
+		});
+        
 	}
 
 	static bool GetEncoderInfo(FNVENCCommon& NVENC, FVideoEncoderInfo& EncoderInfo);
@@ -182,7 +191,11 @@ namespace AVEncoder
 		}
 	}
 
-	FVideoEncoderNVENC_H264::FVideoEncoderNVENC_H264() : NVENC(FNVENCCommon::Setup()) {}
+	FVideoEncoderNVENC_H264::FVideoEncoderNVENC_H264() : NVENC(FNVENCCommon::Setup()) 
+	{
+		// Parse NVENC settings from command line (if any relevant ones are passed)
+		NVENCParseCommandLineFlags();
+	}
 
 	FVideoEncoderNVENC_H264::~FVideoEncoderNVENC_H264() { Shutdown(); }
 
@@ -223,9 +236,6 @@ namespace AVEncoder
 		{
 			mutableConfig.MaxFramerate = 60;
 		}
-		
-		// Parse NVENC settings from command line (if any relevant ones are passed)
-		NVENCParseCommandLineFlags();
 
 		return AddLayer(mutableConfig);
 	}
@@ -245,29 +255,29 @@ namespace AVEncoder
 
 	void FVideoEncoderNVENC_H264::Encode(FVideoEncoderInputFrame const* InFrame, const FEncodeOptions& EncodeOptions)
 	{
-		for(auto&& layer : Layers)
+		for(FVideoEncoder::FLayer* Layer : Layers)
 		{
-			auto const nvencLayer = static_cast<FNVENCLayer*>(layer);
-			nvencLayer->Encode(InFrame, EncodeOptions);
+			FVideoEncoderNVENC_H264::FNVENCLayer* NVLayer = static_cast<FVideoEncoderNVENC_H264::FNVENCLayer*>(Layer);
+			NVLayer->Encode(InFrame, EncodeOptions);
 		}
 	}
 
 	void FVideoEncoderNVENC_H264::Flush()
 	{
-		for(auto&& layer : Layers)
+		for(FVideoEncoder::FLayer* Layer : Layers)
 		{
-			auto const nvencLayer = static_cast<FNVENCLayer*>(layer);
-			nvencLayer->Flush();
+			FVideoEncoderNVENC_H264::FNVENCLayer* NVLayer = static_cast<FVideoEncoderNVENC_H264::FNVENCLayer*>(Layer);
+			NVLayer->Flush();
 		}
 	}
 
 	void FVideoEncoderNVENC_H264::Shutdown()
 	{
-		for(auto&& layer : Layers)
+		for(FVideoEncoder::FLayer* Layer : Layers)
 		{
-			auto const nvencLayer = static_cast<FNVENCLayer*>(layer);
-			nvencLayer->Shutdown();
-			DestroyLayer(nvencLayer);
+			FVideoEncoderNVENC_H264::FNVENCLayer* NVLayer = static_cast<FVideoEncoderNVENC_H264::FNVENCLayer*>(Layer);
+			NVLayer->Shutdown();
+			DestroyLayer(NVLayer);
 		}
 		Layers.Reset();
 	}
@@ -285,14 +295,14 @@ namespace AVEncoder
 		if(CreateSession() && CreateInitialConfig())
 		{
 			// create encoder
-			auto const result = NVENC.nvEncInitializeEncoder(NVEncoder, &EncoderInitParams);
-			if(result != NV_ENC_SUCCESS)
+			NVENCSTATUS Result = NVENC.nvEncInitializeEncoder(NVEncoder, &EncoderInitParams);
+			if(Result != NV_ENC_SUCCESS)
 			{
-				UE_LOG(LogEncoderNVENC, Error, TEXT("Unable to initialize NvEnc encoder (%s)."), *NVENC.GetErrorString(NVEncoder, result));
+				UE_LOG(LogEncoderNVENC, Error, TEXT("Unable to initialize NvEnc encoder (%s)."), *NVENC.GetErrorString(NVEncoder, Result));
+				return false;
 			}
 			else
 			{
-				EncoderThread = MakeUnique<FThread>(TEXT("NvencEncoderLayerThread"), [this]() { ProcessFramesFunc(); });
 				return true;
 			}
 		}
@@ -379,8 +389,8 @@ namespace AVEncoder
 		/*
 		* Intra refresh - used to stabilise stream on the decoded side when frames are dropped/lost.
 		*/
-		int32 IntraRefreshPeriodFrames = CVarIntraRefreshPeriodFrames.GetValueOnAnyThread();
-		int32 IntraRefreshCountFrames = CVarIntraRefreshCountFrames.GetValueOnAnyThread();
+		int32 IntraRefreshPeriodFrames = CVarNVENCIntraRefreshPeriodFrames.GetValueOnAnyThread();
+		int32 IntraRefreshCountFrames = CVarNVENCIntraRefreshCountFrames.GetValueOnAnyThread();
 		bool bIntraRefreshSupported = GetEncoderCapability(NVENC, NVEncoder, NV_ENC_CAPS_SUPPORT_INTRA_REFRESH) > 0; 
 		bool bIntraRefreshEnabled = IntraRefreshPeriodFrames > 0;
 
@@ -402,7 +412,7 @@ namespace AVEncoder
 		/*
 		* IDR period - how often to send IDR (instantaneous decode refresh) frames, a.k.a keyframes. This can stabilise a stream that dropped/lost some frames (but at the cost of more bandwidth).
 		*/
-		int32 IdrPeriod = CVarKeyframeInterval.GetValueOnAnyThread();
+		int32 IdrPeriod = CVarNVENCKeyframeInterval.GetValueOnAnyThread();
 		if(IdrPeriod > 0)
 		{
 			EncoderConfig.encodeCodecConfig.h264Config.idrPeriod = IdrPeriod;
@@ -420,12 +430,20 @@ namespace AVEncoder
 		EncoderConfig.encodeCodecConfig.h264Config.sliceMode = 0;
 		EncoderConfig.encodeCodecConfig.h264Config.sliceModeData = 0;
 
+		/*
+		* These put extra meta data into the frames that allow Firefox to join mid stream.
+		*/
+		EncoderConfig.encodeCodecConfig.h264Config.outputFramePackingSEI = 1;
+		EncoderConfig.encodeCodecConfig.h264Config.outputRecoveryPointSEI = 1;
+
 		UpdateConfig();
 		return true;
 	}
 
 	void FVideoEncoderNVENC_H264::FNVENCLayer::MaybeReconfigure()
 	{
+		uint64 PreReconfigureCycles = FPlatformTime::Cycles64();
+
 		FScopeLock lock(&ConfigMutex);
 		if(NeedsReconfigure)
 		{
@@ -444,6 +462,15 @@ namespace AVEncoder
 				UE_LOG(LogEncoderNVENC, Error, TEXT("Failed to update NVENC encoder configuration (%s)"), *NVENC.GetErrorString(NVEncoder, result));
 
 			NeedsReconfigure = false;
+
+			uint64 PostReconfigureCycles = FPlatformTime::Cycles64();
+
+			if(CVarNVENCEnableStats.GetValueOnAnyThread())
+			{
+				// The time it took to do nvEncEncodePicture
+				double ReconfigureDeltaMs = FPlatformTime::ToMilliseconds64(PostReconfigureCycles - PreReconfigureCycles);
+				FNVENCStats::Get().SetReconfigureLatency(ReconfigureDeltaMs);
+			}
 		}
 	}
 
@@ -479,7 +506,7 @@ namespace AVEncoder
 
 		// If we have QP ranges turned on use the last encoded QP to guide the max QP for an i-frame, so the i-frame doesn't look too blocky
 		// Note: this does nothing if we have i-frames turned off.
-		if(RateControlParams.enableMaxQP && LastEncodedQP > 0 && CVarKeyframeQPUseLastQP.GetValueOnAnyThread())
+		if(RateControlParams.enableMaxQP && LastEncodedQP > 0 && CVarNVENCKeyframeQPUseLastQP.GetValueOnAnyThread())
 		{
 			RateControlParams.maxQP.qpIntra = LastEncodedQP;
 		}
@@ -496,230 +523,239 @@ namespace AVEncoder
 		}
 	}
 
-	void FVideoEncoderNVENC_H264::FNVENCLayer::Encode(FVideoEncoderInputFrame const* InFrame, const FEncodeOptions& EncodeOptions)
+	void FVideoEncoderNVENC_H264::FNVENCLayer::EncodeBuffer(FInputOutput* Buffer)
 	{
-		uint64 EncodeStartMs = FPlatformTime::ToMilliseconds64(FPlatformTime::Cycles64());
-		FInputOutput* Buffer = GetOrCreateBuffer(static_cast<const FVideoEncoderInputFrameImpl*>(InFrame->Obtain()));
+		checkf(Buffer, TEXT("Cannot encode null buffer."));
+		if(Buffer == nullptr)
+		{
+			UE_LOG(LogEncoderNVENC, Error, TEXT("Cannot proceed with encoding in NVENC - buffer was nullptr."));
+			return;
+		}
 
-		if(Buffer)
+		if (Buffer && Buffer->PicParams.encodePicFlags == NV_ENC_PIC_FLAG_EOS)
 		{
-			Buffer->EncodeStartMs = EncodeStartMs;
+			NVENCSTATUS Result = NVENC.nvEncEncodePicture(NVEncoder, &(Buffer->PicParams));
+			DestroyBuffer(Buffer);
+			return;
+		}
 
-			if(MapInputTexture(*Buffer))
-			{
-				Buffer->PicParams = {};
-				Buffer->PicParams.version = NV_ENC_PIC_PARAMS_VER;
-				Buffer->PicParams.inputWidth = Buffer->Width;
-				Buffer->PicParams.inputHeight = Buffer->Height;
-				Buffer->PicParams.inputPitch = Buffer->Pitch ? Buffer->Pitch : Buffer->Width;
-				Buffer->PicParams.inputBuffer = Buffer->MappedInput;
-				Buffer->PicParams.bufferFmt = Buffer->BufferFormat;
-				Buffer->PicParams.encodePicFlags = 0;
-
-				if(EncodeOptions.bForceKeyFrame)
-				{
-					LastKeyFrameTime = FDateTime::UtcNow();
-					Buffer->PicParams.encodePicFlags |= NV_ENC_PIC_FLAG_FORCEIDR;
-				}
+		uint64 StartQueueEncodeCycles = FPlatformTime::Cycles64();
+
+		if(Buffer && !RegisterInputTexture(Buffer))
+		{
+			UE_LOG(LogEncoderNVENC, Error, TEXT("Unable to register texture in NVENC"));
+			DestroyBuffer(Buffer);
+			return;
+		}
+		
+		Buffer->SubmitTimeCycles = StartQueueEncodeCycles;
+		Buffer->PicParams.version = NV_ENC_PIC_PARAMS_VER;
+		Buffer->PicParams.inputWidth = Buffer->Width;
+		Buffer->PicParams.inputHeight = Buffer->Height;
+		Buffer->PicParams.inputPitch = Buffer->Pitch ? Buffer->Pitch : Buffer->Width;
+		Buffer->PicParams.inputBuffer = Buffer->MappedInput;
+		Buffer->PicParams.bufferFmt = Buffer->BufferFormat;
+		Buffer->PicParams.outputBitstream = Buffer->OutputBitstream;
+		Buffer->PicParams.pictureStruct = NV_ENC_PIC_STRUCT_FRAME;
+
+		MaybeReconfigure();
+
+		// check if we have not reconfigued the buffer size
+		if(Buffer->SourceFrame && (Buffer->Width != EncoderInitParams.encodeWidth || Buffer->Height != EncoderInitParams.encodeHeight))
+		{
+			DestroyBuffer(Buffer);
+			return;
+		}
+
+		// Do synchronous encode
+		uint64 PrenvEncEncodePictureCycles = FPlatformTime::Cycles64();
+		NVENCSTATUS Result = NVENC.nvEncEncodePicture(NVEncoder, &(Buffer->PicParams));
 
-				Buffer->PicParams.inputTimeStamp = Buffer->TimeStamp = InFrame->GetTimestampUs();
-				Buffer->PicParams.outputBitstream = Buffer->OutputBitstream;
-				Buffer->PicParams.pictureStruct = NV_ENC_PIC_STRUCT_FRAME;
-				Buffer->PicParams.frameIdx = InFrame->GetFrameID();
+		uint64 PostnvEncEncodePictureCycles = FPlatformTime::Cycles64();
+		if(CVarNVENCEnableStats.GetValueOnAnyThread())
+		{
+			// The time it took to do nvEncEncodePicture
+			double ProcessFramesFuncDeltaMs = FPlatformTime::ToMilliseconds64(PostnvEncEncodePictureCycles - PrenvEncEncodePictureCycles);
+			FNVENCStats::Get().SetnvEncEncodePictureLatency(ProcessFramesFuncDeltaMs);
+		}
+
+		// We have an encoded bitstream - to use it we must bitstream lock, read, unlock (and then pass out the encoded output the whoever needs it)
+		if(Result == NV_ENC_SUCCESS)
+		{
+			ProcessEncodedBuffer(Buffer);
+
+			UnregisterInputTexture(Buffer);
+
+			if(CVarNVENCEnableStats.GetValueOnAnyThread())
+			{
+				uint64 EndQueueEncodeCycles = FPlatformTime::Cycles64();
+				double QueueEncodeMs = FPlatformTime::ToMilliseconds64(EndQueueEncodeCycles - StartQueueEncodeCycles);
+				FNVENCStats::Get().SetQueueEncodeLatency(QueueEncodeMs);
+			}
 
-				PendingEncodes.Enqueue(Buffer);
-				FramesPending->Trigger();
+			return;
+		}
+		else
+		{
+			// Something went wrong
+			UE_LOG(LogEncoderNVENC, Error, TEXT("nvEncEncodePicture returned %s"), *NVENC.GetErrorString(NVEncoder, Result));
+			
+			if(Buffer)
+			{
+				DestroyBuffer(Buffer);
 			}
 		}
 	}
 
-	void FVideoEncoderNVENC_H264::FNVENCLayer::Flush()
+	void FVideoEncoderNVENC_H264::FNVENCLayer::Encode(FVideoEncoderInputFrame const* InFrame, const FEncodeOptions& EncodeOptions)
 	{
-		FInputOutput* EmptyBuffer = CreateBuffer();
+		FInputOutput* Buffer = GetOrCreateBuffer(static_cast<const FVideoEncoderInputFrameImpl*>(InFrame));
 
-		if(!EmptyBuffer)
+		if(!Buffer)
 		{
+			UE_LOG(LogEncoderNVENC, Error, TEXT("Cannot encode null buffer."));
 			return;
 		}
 
-		EmptyBuffer->PicParams = {};
-		EmptyBuffer->PicParams.version = NV_ENC_PIC_PARAMS_VER;
-		EmptyBuffer->PicParams.encodePicFlags = NV_ENC_PIC_FLAG_EOS;
+		// Clear pic params from last time this buffer was used
+		Buffer->PicParams = {};
+		Buffer->PicParams.encodePicFlags = 0;
+
+		if(EncodeOptions.bForceKeyFrame)
+		{
+			LastKeyFrameTime = FDateTime::UtcNow();
+			Buffer->PicParams.encodePicFlags |= NV_ENC_PIC_FLAG_FORCEIDR;
+		}
+
+		// Send SPS/PPS every frame (wasteful)
+		//Buffer->PicParams.encodePicFlags |= NV_ENC_PIC_FLAG_OUTPUT_SPSPPS;
 
-		PendingEncodes.Enqueue(EmptyBuffer);
+		Buffer->PicParams.frameIdx = InFrame->GetFrameID();
+		Buffer->PicParams.inputTimeStamp = Buffer->TimeStamp = InFrame->GetTimestampUs();
+
+		EncodeBuffer(Buffer);
+	}
+
+	void FVideoEncoderNVENC_H264::FNVENCLayer::Flush()
+	{
+		// Empty the FrameToBufferMapping
+
+		// As destroy buffer also modifies the frames to buffer mapping we must remove entries by copying keys first.
+		TArray<const AVEncoder::FVideoEncoderInputFrame*> Frames;
+		FrameToBufferMapping.GetKeys(Frames);
+		for(const AVEncoder::FVideoEncoderInputFrame* StoredFrame : Frames)
+		{
+			FInputOutput* Buffer = *FrameToBufferMapping.Find(StoredFrame);
+			DestroyBuffer(Buffer);
+		}
+
+		Frames.Empty();
+		FrameToBufferMapping.Empty();
 	}
 
 	void FVideoEncoderNVENC_H264::FNVENCLayer::Shutdown()
 	{
 		Flush();
 
-		
-		FramesPending->Trigger();
-		EncoderThread->Join();
+		// Send an empty buffer through the encoder with the special "End of Stream (EOS)" flag set to shutdown properly.
+		{
+			FInputOutput* EmptyBuffer = CreateBuffer();
+
+			if(!EmptyBuffer)
+			{
+				return;
+			}
+
+			EmptyBuffer->PicParams = {};
+			EmptyBuffer->PicParams.version = NV_ENC_PIC_PARAMS_VER;
+			EmptyBuffer->PicParams.encodePicFlags = NV_ENC_PIC_FLAG_EOS;
+
+			EncodeBuffer(EmptyBuffer);
+		}
 
 		if(NVEncoder)
 		{
+			UE_LOG(LogEncoderNVENC, Log, TEXT("Attempting to destroy NVENC encoder."));
+
 			auto const result = NVENC.nvEncDestroyEncoder(NVEncoder);
 			if(result != NV_ENC_SUCCESS)
 			{
 				UE_LOG(LogEncoderNVENC, Error, TEXT("Failed to destroy NVENC encoder (%s)."), *NVENC.GetErrorString(NVEncoder, result));
 			}
+			else
+			{
+				UE_LOG(LogEncoderNVENC, Log, TEXT("Successfully destroyed NVENC encoder."));
+			}
 			NVEncoder = nullptr;
 		}
 	}
 
-// This function throws error in static analysis as the Buffer is never explicitly set
-// however dequeue will guarantee that it is never null and hence a false positive
-#pragma warning( push )
-#pragma warning( disable : 6011 )
-	void FVideoEncoderNVENC_H264::FNVENCLayer::ProcessFramesFunc()
+
+	void FVideoEncoderNVENC_H264::FNVENCLayer::ProcessEncodedBuffer(FInputOutput* Buffer)
 	{
-		TQueue<FInputOutput*> LiveBuffers;
-		FInputOutput* Buffer;
-		NVENCSTATUS result;
+		uint64 StartProcessFramesFuncCycles;
+		uint64 FinishedEncodingCycles;
+		FCodecPacketImpl Packet;
 
-		while(bShouldEncoderThreadRun)
-		{
-			FramesPending->Wait();
+		StartProcessFramesFuncCycles = FPlatformTime::Cycles64();
 
-			while (bShouldEncoderThreadRun && PendingEncodes.Dequeue(Buffer))
+		// lock output buffers for CPU access
+		if(LockOutputBuffer(Buffer))
+		{
+			if(Encoder.OnEncodedPacket)
 			{
-				// we have kicked off a flush of the encoder clear all the buffers
-				if (Buffer->PicParams.encodePicFlags == NV_ENC_PIC_FLAG_EOS)
+				Packet.Data = static_cast<const uint8*>(Buffer->BitstreamData);
+				Packet.DataSize = Buffer->BitstreamDataSize;
+
+				if(Buffer->PictureType & NV_ENC_PIC_TYPE_IDR)
 				{
-					FInputOutput* EOSBuffer = Buffer;
-					DestroyBuffer(EOSBuffer);
-
-					while (IdleBuffers.Dequeue(Buffer))
-					{
-						FInputOutput* OldBuffer = Buffer;
-						DestroyBuffer(OldBuffer);
-					}
-
-					// We are exiting so clear all live buffers
-					while (LiveBuffers.Dequeue(Buffer))
-					{
-						FInputOutput* LiveBuffer = Buffer;
-						if (LiveBuffer->SourceFrame)
-						{
-							LiveBuffer->SourceFrame->Release();
-							LiveBuffer->SourceFrame = nullptr;
-						}
-						DestroyBuffer(LiveBuffer);
-					}
-
-					bShouldEncoderThreadRun = false;
-					return;
+					UE_LOG(LogEncoderNVENC, Verbose, TEXT("Generated IDR Frame"));
+					Packet.IsKeyFrame = true;
 				}
-
-
-				MaybeReconfigure();
-
-				// check if we have not reconfigued the buffer size
-				if(Buffer->Width != EncoderInitParams.encodeWidth || Buffer->Height != EncoderInitParams.encodeHeight)
+				else
 				{
-					// clear all live buffers as they are now invalid
-					LiveBuffers.Enqueue(Buffer);
-
-					while(LiveBuffers.Dequeue(Buffer))
-					{
-						Buffer->SourceFrame->Release();
-						Buffer->SourceFrame = nullptr;
-						DestroyBuffer(Buffer);
-					}
-
-					continue;
+					// If it is not a keyframe store the QP.
+					UpdateLastEncodedQP(Buffer->FrameAvgQP);
 				}
 
-				// Do synchronous encode
-				result = NVENC.nvEncEncodePicture(NVEncoder, &(Buffer->PicParams));
+				FinishedEncodingCycles = FPlatformTime::Cycles64();
+				Packet.VideoQP = Buffer->FrameAvgQP;
+				Packet.Timings.StartTs = FTimespan::FromMilliseconds(FPlatformTime::ToMilliseconds64(Buffer->SubmitTimeCycles));
+				Packet.Timings.FinishTs = FTimespan::FromMilliseconds(FPlatformTime::ToMilliseconds64(FinishedEncodingCycles));
+				Packet.Framerate = EncoderInitParams.frameRateNum;
 
-				// Enqueue buffers until we get a success then Dequeue all the LiveBuffers
-				if(result == NV_ENC_ERR_NEED_MORE_INPUT)
+				// Stats
+				if(CVarNVENCEnableStats.GetValueOnAnyThread())
 				{
-					LiveBuffers.Enqueue(Buffer);
-				}
-				else if(result == NV_ENC_SUCCESS)
-				{
-					LiveBuffers.Enqueue(Buffer);
-
-					while(bShouldEncoderThreadRun && LiveBuffers.Dequeue(Buffer))
-					{
-						// lock output buffers for CPU access
-						if(LockOutputBuffer(*Buffer))
-						{
-							if(Encoder.OnEncodedPacket)
-							{
-								// create packet with buffer contents
-								FCodecPacketImpl Packet;
-
-								Packet.Data = static_cast<const uint8*>(Buffer->BitstreamData);
-								Packet.DataSize = Buffer->BitstreamDataSize;
-								
-								if(Buffer->PictureType & NV_ENC_PIC_TYPE_IDR)
-								{
-									UE_LOG(LogEncoderNVENC, Verbose, TEXT("Generated IDR Frame"));
-									Packet.IsKeyFrame = true;
-								}
-								else
-								{
-									// If it is not a keyframe store the QP.
-									UpdateLastEncodedQP(Buffer->FrameAvgQP);
-								}
-
-								Packet.VideoQP = Buffer->FrameAvgQP;
-								Packet.Timings.StartTs = FTimespan::FromMilliseconds(Buffer->EncodeStartMs);
-								Packet.Timings.FinishTs = FTimespan::FromMilliseconds(FPlatformTime::ToMilliseconds64(FPlatformTime::Cycles64()));
-								Packet.Framerate = EncoderInitParams.frameRateNum;
-
-								// yeet it out
-								if(Encoder.OnEncodedPacket)
-								{
-									Encoder.OnEncodedPacket(LayerIndex, Buffer->SourceFrame, Packet);
-								}
-							}
-
-							UnlockOutputBuffer(*Buffer);
-						}
-
-
-
-						// We have consumed the encoded buffer we can now return it to the pool
-						if(Buffer->SourceFrame)
-						{
-							Buffer->SourceFrame->Release();
-							Buffer->SourceFrame = nullptr;
-						}
-
-						IdleBuffers.Enqueue(Buffer);
-					}
+					// The time it took to do just the encode part
+					double ProcessFramesFuncDeltaMs = FPlatformTime::ToMilliseconds64(FinishedEncodingCycles - StartProcessFramesFuncCycles);
+					FNVENCStats::Get().SetProcessFramesFuncLatency(ProcessFramesFuncDeltaMs);
+
+					// Total time the encoder took to process a frame from submit through to encode
+					double TotalEncoderLatencyMs = Packet.Timings.FinishTs.GetTotalMilliseconds() - Packet.Timings.StartTs.GetTotalMilliseconds();
+					FNVENCStats::Get().SetTotalEncoderLatency(TotalEncoderLatencyMs);
 				}
-				else
+
+				// Send frame to OnEncodedPacket listener, such as WebRTC.
+				if(Encoder.OnEncodedPacket)
 				{
-					// Something went wrong
-					UE_LOG(LogEncoderNVENC, Error, TEXT("nvEncEncodePicture returned %s"), *NVENC.GetErrorString(NVEncoder, result));
-
-					if(Buffer->SourceFrame)
-					{
-						Buffer->SourceFrame->Release();
-						Buffer->SourceFrame = nullptr;
-						DestroyBuffer(Buffer);
-					}
+					Encoder.OnEncodedPacket(LayerIndex, Buffer->SourceFrame, Packet);
 				}
 			}
-		}
 
-		// We are exiting so clear all live buffers
-		while(LiveBuffers.Dequeue(Buffer))
-		{
-			if (Buffer->SourceFrame)
+			if(!UnlockOutputBuffer(Buffer))
 			{
-				Buffer->SourceFrame->Release();
-				Buffer->SourceFrame = nullptr;
+				UE_LOG(LogEncoderNVENC, Error, TEXT("Failed to unlock output buffer."));
+				DestroyBuffer(Buffer);
 			}
+
+		}
+		else
+		{
+			UE_LOG(LogEncoderNVENC, Error, TEXT("Failed to lock output buffer."));
+			DestroyBuffer(Buffer);
 		}
 	}
-#pragma warning( pop )
 
 	int FVideoEncoderNVENC_H264::FNVENCLayer::GetCapability(NV_ENC_CAPS CapsToQuery) const
 	{
@@ -737,80 +773,35 @@ namespace AVEncoder
 
 	FVideoEncoderNVENC_H264::FNVENCLayer::FInputOutput* FVideoEncoderNVENC_H264::FNVENCLayer::GetOrCreateBuffer(const FVideoEncoderInputFrameImpl* InFrame)
 	{
-		void* TextureToCompress = nullptr;
-
-		switch(InFrame->GetFormat())
-		{
-#if PLATFORM_WINDOWS
-		case EVideoFrameFormat::D3D11_R8G8B8A8_UNORM:
-			TextureToCompress = InFrame->GetD3D11().EncoderTexture;
-			break;
-		case EVideoFrameFormat::D3D12_R8G8B8A8_UNORM:
-			UE_LOG(LogEncoderNVENC, Error, TEXT("NVENC got passed a DX12 texture which it does not support, make sure it has been shared with a DX11 context and pass that texture instead."));
-			return nullptr; // NVENC cant encode a D3D12 texture directly this should be converted to D3D11
-#endif
-		case EVideoFrameFormat::CUDA_R8G8B8A8_UNORM:
-			TextureToCompress = InFrame->GetCUDA().EncoderTexture;
-			break;
-		case AVEncoder::EVideoFrameFormat::Undefined:
-		default:
-			break;
-		}
-
-		if(!TextureToCompress)
-		{
-			UE_LOG(LogEncoderNVENC, Fatal, TEXT("Got passed a null texture to encode."));
-			return nullptr;
-		}
-
 		FInputOutput* Buffer = nullptr;
-		FInputOutput* TempBuffer = nullptr;
-		TQueue<FInputOutput*> TempQueue;
-		while(!IdleBuffers.IsEmpty())
-		{
-			IdleBuffers.Dequeue(TempBuffer);
-			if(TempBuffer->InputTexture == TextureToCompress)
-			{
-				Buffer = TempBuffer;
-				break;
-			}
-			else
-			{
-				TempQueue.Enqueue(TempBuffer);
-			}
-		}
-
-		while(!TempQueue.IsEmpty())
-		{
-			TempQueue.Dequeue(TempBuffer);
-			IdleBuffers.Enqueue(TempBuffer);
-		}
-
-		if(Buffer)
-		{
-			// Check for buffer and InputFrameFactory texture resolution mismatch
-			if(InFrame->GetWidth() != Buffer->Width || InFrame->GetHeight() != Buffer->Height)
-			{
-				DestroyBuffer(Buffer);
-				Buffer = nullptr;
-			}
-		}
 
-		if(!Buffer)
+		if(FrameToBufferMapping.Contains(InFrame))
 		{
-			Buffer = CreateBuffer();
-			Buffer->SourceFrame = InFrame;
+			Buffer = *FrameToBufferMapping.Find(InFrame);
+			checkf(Buffer, TEXT("Frame was mapped to a nullptr buffer, this indicate a buffer was not removed when it should have been."));
 
-			if(Buffer && !RegisterInputTexture(*Buffer, TextureToCompress, FIntPoint(InFrame->GetWidth(), InFrame->GetHeight())))
+			if(Buffer)
 			{
-				Buffer->SourceFrame->Release();
-				DestroyBuffer(Buffer);
-				Buffer = nullptr;
+				// Check for buffer and Frame texture resolution mismatch
+				if(InFrame->GetWidth() != Buffer->Width || InFrame->GetHeight() != Buffer->Height)
+				{
+					DestroyBuffer(Buffer);
+					Buffer = nullptr;
+				}
+				// Resolution of existing buffer is a match, so set our frame as the source frame
+				else
+				{
+					Buffer->SourceFrame = InFrame;
+				}
 			}
+			
 		}
+		// No existing buffer for this frame
 		else
 		{
+			Buffer = CreateBuffer();
 			Buffer->SourceFrame = InFrame;
+			FrameToBufferMapping.Add(InFrame, Buffer);
 		}
 
 		return Buffer;
@@ -839,11 +830,24 @@ namespace AVEncoder
 
 	void FVideoEncoderNVENC_H264::FNVENCLayer::DestroyBuffer(FInputOutput* InBuffer)
 	{
-		// unregister input texture - if any
-		UnregisterInputTexture(*InBuffer);
+		checkf(InBuffer, TEXT("Cannot destroy buffer in NVENC - buffer was nullptr."));
+		if(!InBuffer)
+		{
+			return;
+		}
+
+		// Remove from our mapping of frames to buffers
+		if(InBuffer->SourceFrame)
+		{
+			FrameToBufferMapping.Remove(InBuffer->SourceFrame);
+			InBuffer->SourceFrame->Release();
+		}
 
-		// destroy output buffer - if any
-		UnlockOutputBuffer(*InBuffer);
+		// Unregister input texture - if any
+		UnregisterInputTexture(InBuffer);
+
+		// Destroy output buffer - if any
+		UnlockOutputBuffer(InBuffer);
 		if(InBuffer->OutputBitstream)
 		{
 			auto const result = NVENC.nvEncDestroyBitstreamBuffer(NVEncoder, InBuffer->OutputBitstream);
@@ -854,70 +858,118 @@ namespace AVEncoder
 			InBuffer->OutputBitstream = nullptr;
 		}
 
-		// release source texture
+		// Clear source texture
 		InBuffer->InputTexture = nullptr;
 
 		delete InBuffer;
 	}
 
-	void FVideoEncoderNVENC_H264::FNVENCLayer::CreateResourceDIRECTX(FInputOutput& InBuffer, NV_ENC_REGISTER_RESOURCE& RegisterParam, FIntPoint TextureSize)
+	void FVideoEncoderNVENC_H264::FNVENCLayer::CreateResourceDIRECTX(FInputOutput* InBuffer, NV_ENC_REGISTER_RESOURCE& RegisterParam, FIntPoint TextureSize)
 	{
+		checkf(InBuffer, TEXT("Cannot create DirectX resource in NVENC - buffer was nullptr."));
+		if(!InBuffer)
+		{
+			return;
+		}
+
 #if PLATFORM_WINDOWS
 		D3D11_TEXTURE2D_DESC Desc;
-		static_cast<ID3D11Texture2D*>(InBuffer.InputTexture)->GetDesc(&Desc);
+		static_cast<ID3D11Texture2D*>(InBuffer->InputTexture)->GetDesc(&Desc);
 
 		switch(Desc.Format)
 		{
 		case DXGI_FORMAT_NV12:
-			InBuffer.BufferFormat = NV_ENC_BUFFER_FORMAT_NV12;
+			InBuffer->BufferFormat = NV_ENC_BUFFER_FORMAT_NV12;
 			break;
 		case DXGI_FORMAT_R8G8B8A8_UNORM:
-			InBuffer.BufferFormat = NV_ENC_BUFFER_FORMAT_ABGR;
+			InBuffer->BufferFormat = NV_ENC_BUFFER_FORMAT_ABGR;
 			break;
 		case DXGI_FORMAT_B8G8R8A8_UNORM:
-			InBuffer.BufferFormat = NV_ENC_BUFFER_FORMAT_ARGB;
+			InBuffer->BufferFormat = NV_ENC_BUFFER_FORMAT_ARGB;
 			break;
 		default:
 			UE_LOG(LogEncoderNVENC, Error, TEXT("Invalid input texture format for NVENC (%d)"), Desc.Format);
 			return;
 		}
 
-		InBuffer.Width = TextureSize.X;
-		InBuffer.Height = TextureSize.Y;
-		InBuffer.Pitch = 0;
+		InBuffer->Width = TextureSize.X;
+		InBuffer->Height = TextureSize.Y;
+		InBuffer->Pitch = 0;
 
 		RegisterParam.resourceType = NV_ENC_INPUT_RESOURCE_TYPE_DIRECTX;
 		RegisterParam.width = Desc.Width;
 		RegisterParam.height = Desc.Height;
-		RegisterParam.pitch = InBuffer.Pitch;
-		RegisterParam.bufferFormat = InBuffer.BufferFormat;
+		RegisterParam.pitch = InBuffer->Pitch;
+		RegisterParam.bufferFormat = InBuffer->BufferFormat;
 		RegisterParam.bufferUsage = NV_ENC_INPUT_IMAGE;
 #endif
 	}
 
-	void FVideoEncoderNVENC_H264::FNVENCLayer::CreateResourceCUDAARRAY(FInputOutput& InBuffer, NV_ENC_REGISTER_RESOURCE& RegisterParam, FIntPoint TextureSize)
+	void FVideoEncoderNVENC_H264::FNVENCLayer::CreateResourceCUDAARRAY(FInputOutput* InBuffer, NV_ENC_REGISTER_RESOURCE& RegisterParam, FIntPoint TextureSize)
 	{
-		InBuffer.Width = TextureSize.X;
-		InBuffer.Height = TextureSize.Y;
-		InBuffer.Pitch = TextureSize.X * 4;
-		InBuffer.BufferFormat = NV_ENC_BUFFER_FORMAT_ARGB;
+		checkf(InBuffer, TEXT("Cannot create CUDA array in NVENC - buffer was nullptr."));
+		if(!InBuffer)
+		{
+			return;
+		}
+
+		InBuffer->Width = TextureSize.X;
+		InBuffer->Height = TextureSize.Y;
+		InBuffer->Pitch = TextureSize.X * 4;
+		InBuffer->BufferFormat = NV_ENC_BUFFER_FORMAT_ARGB;
 
 		RegisterParam.resourceType = NV_ENC_INPUT_RESOURCE_TYPE_CUDAARRAY;
-		RegisterParam.width = InBuffer.Width;
-		RegisterParam.height = InBuffer.Height;
-		RegisterParam.pitch = InBuffer.Pitch;
-		RegisterParam.bufferFormat = InBuffer.BufferFormat;
+		RegisterParam.width = InBuffer->Width;
+		RegisterParam.height = InBuffer->Height;
+		RegisterParam.pitch = InBuffer->Pitch;
+		RegisterParam.bufferFormat = InBuffer->BufferFormat;
 		RegisterParam.bufferUsage = NV_ENC_INPUT_IMAGE;
 	}
 
-	bool FVideoEncoderNVENC_H264::FNVENCLayer::RegisterInputTexture(FInputOutput& InBuffer, void* InTexture, FIntPoint TextureSize)
+	bool FVideoEncoderNVENC_H264::FNVENCLayer::RegisterInputTexture(FInputOutput* InBuffer)
 	{
-		if(!InBuffer.InputTexture)
+		checkf(InBuffer, TEXT("Cannot register texture in NVENC - buffer was nullptr."));
+		if(!InBuffer)
 		{
-			InBuffer.InputTexture = InTexture;
+			return false;
+		}
+
+		if(!InBuffer->InputTexture && InBuffer->PicParams.encodePicFlags != NV_ENC_PIC_FLAG_EOS)
+		{
+			checkf(InBuffer->SourceFrame, TEXT("Cannot register texture if source frame is null."));
+
+			void* TextureToCompress = nullptr;
+
+			switch(InBuffer->SourceFrame->GetFormat())
+			{
+#if PLATFORM_WINDOWS
+			case EVideoFrameFormat::D3D11_R8G8B8A8_UNORM:
+				TextureToCompress = InBuffer->SourceFrame->GetD3D11().EncoderTexture;
+				break;
+			case EVideoFrameFormat::D3D12_R8G8B8A8_UNORM:
+				UE_LOG(LogEncoderNVENC, Error, TEXT("NVENC got passed a DX12 texture which it does not support, make sure it has been shared with a DX11 context and pass that texture instead."));
+				return false; // NVENC cant encode a D3D12 texture directly this should be converted to D3D11
+#endif
+			case EVideoFrameFormat::CUDA_R8G8B8A8_UNORM:
+				TextureToCompress = InBuffer->SourceFrame->GetCUDA().EncoderTexture;
+				break;
+			case AVEncoder::EVideoFrameFormat::Undefined:
+			default:
+				break;
+			}
+
+			if(!TextureToCompress)
+			{
+				UE_LOG(LogEncoderNVENC, Fatal, TEXT("Got passed a null texture to encode."));
+				return false;
+			}
+
+			InBuffer->InputTexture = TextureToCompress;
 			NVENCStruct(NV_ENC_REGISTER_RESOURCE, RegisterParam);
 
-			switch(InBuffer.SourceFrame->GetFormat())
+			FIntPoint TextureSize(InBuffer->SourceFrame->GetWidth(), InBuffer->SourceFrame->GetHeight());
+
+			switch(InBuffer->SourceFrame->GetFormat())
 			{
 #if PLATFORM_WINDOWS
 			case EVideoFrameFormat::D3D11_R8G8B8A8_UNORM:
@@ -933,7 +985,7 @@ namespace AVEncoder
 				break;
 			}
 
-			RegisterParam.resourceToRegister = InTexture;
+			RegisterParam.resourceToRegister = InBuffer->InputTexture;
 
 			NVENCSTATUS Result = NVENC.nvEncRegisterResource(NVEncoder, &RegisterParam);
 			if(Result != NV_ENC_SUCCESS)
@@ -941,73 +993,117 @@ namespace AVEncoder
 				UE_LOG(LogEncoderNVENC, Error, TEXT("Failed to register input texture with NVENC (%s)"), *NVENC.GetErrorString(NVEncoder, Result));
 				return false;
 			}
-			InBuffer.RegisteredInput = RegisterParam.registeredResource;
+			InBuffer->RegisteredInput = RegisterParam.registeredResource;
+
+			if(!MapInputTexture(InBuffer))
+			{
+				UE_LOG(LogEncoderNVENC, Error, TEXT("Could not map input texture, cannot proceed with encoding this buffer."));
+				return false;
+			}
 		}
 
 		return true;
 	}
 
-	bool FVideoEncoderNVENC_H264::FNVENCLayer::UnregisterInputTexture(FInputOutput& InBuffer)
+	bool FVideoEncoderNVENC_H264::FNVENCLayer::UnregisterInputTexture(FInputOutput* InBuffer)
 	{
-		UnmapInputTexture(InBuffer);
-		if(InBuffer.RegisteredInput)
+		checkf(InBuffer, TEXT("Cannot unregister texture in NVENC - buffer was nullptr."));
+		if(!InBuffer)
+		{
+			return false;
+		}
+
+		if(!UnmapInputTexture(InBuffer))
+		{
+			UE_LOG(LogEncoderNVENC, Error, TEXT("Failed to unmap input texture."));
+		}
+
+		if(InBuffer->RegisteredInput)
 		{
-			NVENCSTATUS Result = NVENC.nvEncUnregisterResource(NVEncoder, InBuffer.RegisteredInput);
+			NVENCSTATUS Result = NVENC.nvEncUnregisterResource(NVEncoder, InBuffer->RegisteredInput);
 			if(Result != NV_ENC_SUCCESS)
 			{
 				UE_LOG(LogEncoderNVENC, Error, TEXT("Failed to unregister input texture with NVENC (%s)"), *NVENC.GetErrorString(NVEncoder, Result));
-				InBuffer.InputTexture = nullptr;
-				InBuffer.RegisteredInput = nullptr;
+				InBuffer->InputTexture = nullptr;
+				InBuffer->RegisteredInput = nullptr;
 				return false;
 			}
-			InBuffer.InputTexture = nullptr;
-			InBuffer.RegisteredInput = nullptr;
+			InBuffer->InputTexture = nullptr;
+			InBuffer->RegisteredInput = nullptr;
 		}
 		return true;
 	}
 
-	bool FVideoEncoderNVENC_H264::FNVENCLayer::MapInputTexture(FInputOutput& InBuffer)
+	bool FVideoEncoderNVENC_H264::FNVENCLayer::MapInputTexture(FInputOutput* InBuffer)
 	{
-		if(!InBuffer.MappedInput)
+		checkf(InBuffer, TEXT("Cannot map texture in NVENC - buffer was nullptr."));
+		if(!InBuffer)
+		{
+			return false;
+		}
+
+		if(!InBuffer->MappedInput)
 		{
 			NVENCStruct(NV_ENC_MAP_INPUT_RESOURCE, MapInputResource);
-			MapInputResource.registeredResource = InBuffer.RegisteredInput;
+			MapInputResource.registeredResource = InBuffer->RegisteredInput;
 			NVENCSTATUS Result = NVENC.nvEncMapInputResource(NVEncoder, &MapInputResource);
 			if(Result != NV_ENC_SUCCESS)
 			{
 				UE_LOG(LogEncoderNVENC, Error, TEXT("Failed to map input texture buffer (%s)"), *NVENC.GetErrorString(NVEncoder, Result));
 				return false;
 			}
-			InBuffer.MappedInput = MapInputResource.mappedResource;
-			check(InBuffer.BufferFormat == MapInputResource.mappedBufferFmt);
+			InBuffer->MappedInput = MapInputResource.mappedResource;
+			check(InBuffer->BufferFormat == MapInputResource.mappedBufferFmt);
 		}
 		return true;
 	}
 
-	bool FVideoEncoderNVENC_H264::FNVENCLayer::UnmapInputTexture(FInputOutput& InBuffer)
+	bool FVideoEncoderNVENC_H264::FNVENCLayer::UnmapInputTexture(FInputOutput* InBuffer)
 	{
-		if(InBuffer.MappedInput)
+		checkf(InBuffer, TEXT("Cannot unmap texture in NVENC - buffer was nullptr."));
+		if(!InBuffer)
+		{
+			return false;
+		}
+
+		if(InBuffer->MappedInput)
 		{
-			NVENCSTATUS Result = NVENC.nvEncUnmapInputResource(NVEncoder, InBuffer.MappedInput);
+			NVENCSTATUS Result = NVENC.nvEncUnmapInputResource(NVEncoder, InBuffer->MappedInput);
 			if(Result != NV_ENC_SUCCESS)
 			{
 				UE_LOG(LogEncoderNVENC, Error, TEXT("Failed to unmap input texture buffer (%s)"), *NVENC.GetErrorString(NVEncoder, Result));
-				InBuffer.MappedInput = nullptr;
+				InBuffer->MappedInput = nullptr;
 				return false;
 			}
-			InBuffer.MappedInput = nullptr;
+			InBuffer->MappedInput = nullptr;
 		}
 		return true;
 	}
 
-	bool FVideoEncoderNVENC_H264::FNVENCLayer::LockOutputBuffer(FInputOutput& InBuffer)
+	bool FVideoEncoderNVENC_H264::FNVENCLayer::LockOutputBuffer(FInputOutput* InBuffer)
 	{
-		if(!InBuffer.BitstreamData)
+		checkf(InBuffer, TEXT("Cannot lock output buffer in NVENC - buffer was nullptr."));
+		if(!InBuffer)
+		{
+			return false;
+		}
+
+		if(!InBuffer->BitstreamData)
 		{
+			uint64 PreLockCycles = FPlatformTime::Cycles64();
+
 			// lock output buffers for CPU access
 			NVENCStruct(NV_ENC_LOCK_BITSTREAM, LockBitstreamParam);
-			LockBitstreamParam.outputBitstream = InBuffer.OutputBitstream;
+			LockBitstreamParam.outputBitstream = InBuffer->OutputBitstream;
 			NVENCSTATUS Result = NVENC.nvEncLockBitstream(NVEncoder, &LockBitstreamParam);
+
+			if(CVarNVENCEnableStats.GetValueOnAnyThread())
+			{
+				uint64 PostLockCycles = FPlatformTime::Cycles64();
+				double LockLatencyMs = FPlatformTime::ToMilliseconds64(PostLockCycles - PreLockCycles);
+				FNVENCStats::Get().SetLockOutputBufferLatency(LockLatencyMs);
+			}
+
 			if(Result != NV_ENC_SUCCESS)
 			{
 				UE_LOG(LogEncoderNVENC, Error, TEXT("Failed to lock output bitstream for NVENC encoder (%s)."), *NVENC.GetErrorString(NVEncoder, Result));
@@ -1015,21 +1111,37 @@ namespace AVEncoder
 			}
 			else
 			{
-				InBuffer.BitstreamData = LockBitstreamParam.bitstreamBufferPtr;
-				InBuffer.BitstreamDataSize = LockBitstreamParam.bitstreamSizeInBytes;
-				InBuffer.PictureType = LockBitstreamParam.pictureType;
-				InBuffer.FrameAvgQP = LockBitstreamParam.frameAvgQP;
-				InBuffer.TimeStamp = LockBitstreamParam.outputTimeStamp;
+				InBuffer->BitstreamData = LockBitstreamParam.bitstreamBufferPtr;
+				InBuffer->BitstreamDataSize = LockBitstreamParam.bitstreamSizeInBytes;
+				InBuffer->PictureType = LockBitstreamParam.pictureType;
+				InBuffer->FrameAvgQP = LockBitstreamParam.frameAvgQP;
+				InBuffer->TimeStamp = LockBitstreamParam.outputTimeStamp;
 			}
 		}
 		return true;
 	}
 
-	bool FVideoEncoderNVENC_H264::FNVENCLayer::UnlockOutputBuffer(FInputOutput& InBuffer)
+	bool FVideoEncoderNVENC_H264::FNVENCLayer::UnlockOutputBuffer(FInputOutput* InBuffer)
 	{
-		if(InBuffer.BitstreamData)
+		checkf(InBuffer, TEXT("Cannot unregister texture NVENC - buffer was nullptr."));
+		if(!InBuffer)
+		{
+			return false;
+		}
+
+		if(InBuffer->BitstreamData)
 		{
-			NVENCSTATUS Result = NVENC.nvEncUnlockBitstream(NVEncoder, InBuffer.OutputBitstream);
+			uint64 PreUnlockCycles = FPlatformTime::Cycles64();
+
+			NVENCSTATUS Result = NVENC.nvEncUnlockBitstream(NVEncoder, InBuffer->OutputBitstream);
+
+			if(CVarNVENCEnableStats.GetValueOnAnyThread())
+			{
+				uint64 PostUnlockCycles = FPlatformTime::Cycles64();
+				double UnlockLatencyMs = FPlatformTime::ToMilliseconds64(PostUnlockCycles - PreUnlockCycles);
+				FNVENCStats::Get().SetUnlockOutputBufferLatency(UnlockLatencyMs);
+			}
+
 			if(Result != NV_ENC_SUCCESS)
 			{
 				UE_LOG(LogEncoderNVENC, Error, TEXT("Failed to unlock output bitstream for NVENC encoder (%s)."), *NVENC.GetErrorString(NVEncoder, Result));
@@ -1037,8 +1149,8 @@ namespace AVEncoder
 			}
 			else
 			{
-				InBuffer.BitstreamData = nullptr;
-				InBuffer.BitstreamDataSize = 0;
+				InBuffer->BitstreamData = nullptr;
+				InBuffer->BitstreamDataSize = 0;
 			}
 		}
 		return true;
diff --git a/Engine/Plugins/Media/HardwareEncoders/Source/EncoderNVENC/Private/NVENC_EncoderH264.h b/Engine/Plugins/Media/HardwareEncoders/Source/EncoderNVENC/Private/NVENC_EncoderH264.h
--- a/Engine/Plugins/Media/HardwareEncoders/Source/EncoderNVENC/Private/NVENC_EncoderH264.h
+++ b/Engine/Plugins/Media/HardwareEncoders/Source/EncoderNVENC/Private/NVENC_EncoderH264.h
@@ -20,10 +20,10 @@
         bool Setup(TSharedRef<FVideoEncoderInput> InputFrameFactory, const FLayerConfig& InitConfig) override;
         void Shutdown() override;
 
-        // query whether or not encoder is supported and available
+        // Query whether or not encoder is supported and available
         static bool GetIsAvailable(const FVideoEncoderInput& InFrameFactory, FVideoEncoderInfo &OutEncoderInfo);
 
-        // register encoder with video encoder factory
+        // Register encoder with video encoder factory
         static void Register(FVideoEncoderFactory &InFactory);
 
         void Encode(FVideoEncoderInputFrame const* InFrame, const FEncodeOptions& EncodeOptions) override;
@@ -40,9 +40,29 @@
         {
 
         public:
+
+            struct FInputOutput
+            {
+                const AVEncoder::FVideoEncoderInputFrameImpl* SourceFrame = nullptr;
+                void*  InputTexture = nullptr;
+                uint32 Width = 0;
+                uint32 Height = 0;
+                uint32 Pitch = 0;
+                NV_ENC_BUFFER_FORMAT BufferFormat = NV_ENC_BUFFER_FORMAT_UNDEFINED;
+                NV_ENC_REGISTERED_PTR RegisteredInput = nullptr;
+                NV_ENC_INPUT_PTR MappedInput = nullptr;
+                NV_ENC_PIC_PARAMS PicParams = {};
+                NV_ENC_OUTPUT_PTR OutputBitstream = nullptr;
+                const void *BitstreamData = nullptr;
+                uint32 BitstreamDataSize = 0;
+                NV_ENC_PIC_TYPE PictureType = NV_ENC_PIC_TYPE_UNKNOWN;
+                uint32 FrameAvgQP = 0;
+                uint64 TimeStamp;
+                uint64 SubmitTimeCycles;
+            };
+
             FNVENCLayer(uint32 layerIdx, FLayerConfig const& config, FVideoEncoderNVENC_H264& encoder);
             ~FNVENCLayer();
-
             bool Setup();
             bool CreateSession();
             bool CreateInitialConfig();
@@ -52,66 +72,36 @@
 			void UpdateConfig();
             void UpdateLastEncodedQP(uint32 InLastEncodedQP);
             void Encode(FVideoEncoderInputFrame const* InFrame, const FEncodeOptions& EncodeOptions);
-            void ProcessFramesFunc();
+            void EncodeBuffer(FInputOutput* Buffer);
+            void ProcessEncodedBuffer(FInputOutput* Buffer);
             void Flush();
             void Shutdown();
             void UpdateBitrate(uint32 InMaxBitRate, uint32 InTargetBitRate);
             void UpdateResolution(uint32 InMaxBitRate, uint32 InTargetBitRate);
+            FInputOutput* GetOrCreateBuffer(const FVideoEncoderInputFrameImpl* InFrame);
+            FInputOutput* CreateBuffer();
+            void DestroyBuffer(FInputOutput* InBuffer);
+            bool RegisterInputTexture(FInputOutput* InBuffer);
+            bool UnregisterInputTexture(FInputOutput* InBuffer);
+            bool MapInputTexture(FInputOutput* InBuffer);
+            bool UnmapInputTexture(FInputOutput* InBuffer);
+            bool LockOutputBuffer(FInputOutput* InBuffer);
+            bool UnlockOutputBuffer(FInputOutput* InBuffer);
+            void CreateResourceDIRECTX(FInputOutput* InBuffer, NV_ENC_REGISTER_RESOURCE &RegisterParam, FIntPoint TextureSize);
+            void CreateResourceCUDAARRAY(FInputOutput* InBuffer, NV_ENC_REGISTER_RESOURCE &RegisterParam, FIntPoint TextureSize);
 
+        public:
             FVideoEncoderNVENC_H264 &Encoder;
             FNVENCCommon &NVENC;
             GUID CodecGUID;
             uint32 LayerIndex;
             void *NVEncoder = nullptr;
-            TUniquePtr<FThread> EncoderThread;
-            FThreadSafeBool bShouldEncoderThreadRun = true;
-            FEventRef FramesPending;
-
             NV_ENC_INITIALIZE_PARAMS EncoderInitParams;
             NV_ENC_CONFIG EncoderConfig;
             FDateTime LastKeyFrameTime = 0;
-			bool bForceNextKeyframe = false;
+            bool bForceNextKeyframe = false;
             uint32 LastEncodedQP = 0;
-
-            struct FInputOutput
-            {
-                const AVEncoder::FVideoEncoderInputFrameImpl* SourceFrame = nullptr;
-
-                void*  InputTexture = nullptr;
-                uint32 Width = 0;
-                uint32 Height = 0;
-                uint32 Pitch = 0;
-                NV_ENC_BUFFER_FORMAT BufferFormat = NV_ENC_BUFFER_FORMAT_UNDEFINED;
-                NV_ENC_REGISTERED_PTR RegisteredInput = nullptr;
-                NV_ENC_INPUT_PTR MappedInput = nullptr;
-
-                NV_ENC_PIC_PARAMS PicParams = {};
-
-                NV_ENC_OUTPUT_PTR OutputBitstream = nullptr;
-                const void *BitstreamData = nullptr;
-                uint32 BitstreamDataSize = 0;
-                NV_ENC_PIC_TYPE PictureType = NV_ENC_PIC_TYPE_UNKNOWN;
-                uint32 FrameAvgQP = 0;
-                uint64 TimeStamp;
-
-                uint64 EncodeStartMs;
-            };
-
-            FInputOutput* GetOrCreateBuffer(const FVideoEncoderInputFrameImpl* InFrame);
-            FInputOutput* CreateBuffer();
-            void DestroyBuffer(FInputOutput *InBuffer);
-            bool RegisterInputTexture(FInputOutput &InBuffer, void *InTexture, FIntPoint TextureSize);
-            bool UnregisterInputTexture(FInputOutput &InBuffer);
-            bool MapInputTexture(FInputOutput &InBuffer);
-            bool UnmapInputTexture(FInputOutput &InBuffer);
-            bool LockOutputBuffer(FInputOutput &InBuffer);
-            bool UnlockOutputBuffer(FInputOutput &InBuffer);
-
-            void CreateResourceDIRECTX(FInputOutput &InBuffer, NV_ENC_REGISTER_RESOURCE &RegisterParam, FIntPoint TextureSize);
-            void CreateResourceCUDAARRAY(FInputOutput &InBuffer, NV_ENC_REGISTER_RESOURCE &RegisterParam, FIntPoint TextureSize);
-
-            TQueue<FInputOutput*> IdleBuffers;
-            TQueue<FInputOutput *> PendingEncodes;
+            TMap<const AVEncoder::FVideoEncoderInputFrame*, FInputOutput*> FrameToBufferMapping;
         };
 
         FNVENCCommon &NVENC;
diff --git a/Engine/Plugins/Media/PixelStreaming/Source/PixelStreaming/Private/InputDevice.cpp b/Engine/Plugins/Media/PixelStreaming/Source/PixelStreaming/Private/InputDevice.cpp
--- a/Engine/Plugins/Media/PixelStreaming/Source/PixelStreaming/Private/InputDevice.cpp
+++ b/Engine/Plugins/Media/PixelStreaming/Source/PixelStreaming/Private/InputDevice.cpp
@@ -412,17 +412,7 @@ void FInputDevice::Tick(float DeltaTime)
 	FString Command;
 	while (Commands.Dequeue(Command))
 	{
-		for (UPixelStreamerInputComponent* InputComponent : this->PixelStreamingModule->GetInputComponents())
-		{
-			if (InputComponent->OnCommand(Command))
-			{
-				UE_LOG(PixelStreamerInputDevice, Verbose, TEXT("Command = %s"), *Command);
-			}
-			else
-			{
-				UE_LOG(PixelStreamerInputDevice, Warning, TEXT("Failed to run Command = %s"), *Command);
-			}
-		}
+		UPixelStreamerInputComponent::OnCommand(Command);
 	}
 }
 
diff --git a/Engine/Plugins/Media/PixelStreaming/Source/PixelStreaming/Private/LatencyTester.cpp b/Engine/Plugins/Media/PixelStreaming/Source/PixelStreaming/Private/LatencyTester.cpp
--- a/Engine/Plugins/Media/PixelStreaming/Source/PixelStreaming/Private/LatencyTester.cpp
+++ b/Engine/Plugins/Media/PixelStreaming/Source/PixelStreaming/Private/LatencyTester.cpp
@@ -2,7 +2,7 @@
 
 #include "LatencyTester.h"
 #include "PixelStreamingSettings.h"
-#include <chrono>
+#include "GenericPlatform/GenericPlatformTime.h"
 
 FLatencyTester& FLatencyTester::GetInstance()
 {
@@ -50,11 +50,10 @@ bool FLatencyTester::RecordReceiptTime()
     return false;
 }
 
-bool FLatencyTester::RecordPreCaptureTime(uint32 FrameId)
+bool FLatencyTester::RecordPreCaptureTime()
 {
     if(FLatencyTester::GetInstance().GetTestStage() == FLatencyTester::ELatencyTestStage::PRE_CAPTURE)
     {
-        FLatencyTester::GetInstance().FrameId = FrameId;
         FLatencyTester::GetInstance().PreCaptureTimeMs = FLatencyTester::GetInstance().EpochMillisNow();
         FLatencyTester::GetInstance().TestStage.store(FLatencyTester::ELatencyTestStage::POST_CAPTURE);
         return true;
@@ -64,8 +63,9 @@ bool FLatencyTester::RecordPreCaptureTime(uint32 FrameId)
 
 bool FLatencyTester::RecordPostCaptureTime(uint32 FrameId)
 {
-    if(FLatencyTester::GetInstance().GetTestStage() == FLatencyTester::ELatencyTestStage::POST_CAPTURE && FrameId == FLatencyTester::GetInstance().FrameId)
+    if(FLatencyTester::GetInstance().GetTestStage() == FLatencyTester::ELatencyTestStage::POST_CAPTURE)
     {
+        FLatencyTester::GetInstance().FrameId = FrameId;
         FLatencyTester::GetInstance().PostCaptureTimeMs = FLatencyTester::GetInstance().EpochMillisNow();
         FLatencyTester::GetInstance().TestStage.store(FLatencyTester::ELatencyTestStage::PRE_ENCODE);
         return true;
@@ -103,7 +103,7 @@ bool FLatencyTester::End(FString& OutJSONString, FPlayerId& OutPlayerId)
         FLatencyTester::GetInstance().TransmissionTimeMs = FLatencyTester::GetInstance().EpochMillisNow();
         FLatencyTester::GetInstance().TestStage.store(FLatencyTester::ELatencyTestStage::INACTIVE);
         OutJSONString = FString::Printf( 
-            TEXT( "{ \"ReceiptTimeMs\": %llu, \"PreCaptureTimeMs\": %llu, \"PostCaptureTimeMs\": %llu, \"PreEncodeTimeMs\": %llu, \"PostEncodeTimeMs\": %llu, \"TransmissionTimeMs\": %llu }" ), 
+            TEXT( "{ \"ReceiptTimeMs\": %.2f, \"PreCaptureTimeMs\": %.2f, \"PostCaptureTimeMs\": %.2f, \"PreEncodeTimeMs\": %.2f, \"PostEncodeTimeMs\": %.2f, \"TransmissionTimeMs\": %.2f }" ), 
             FLatencyTester::GetInstance().ReceiptTimeMs,
             FLatencyTester::GetInstance().PreCaptureTimeMs,
             FLatencyTester::GetInstance().PostCaptureTimeMs,
@@ -114,11 +114,11 @@ bool FLatencyTester::End(FString& OutJSONString, FPlayerId& OutPlayerId)
         OutPlayerId = FLatencyTester::GetInstance().PlayerWhoTriggeredTest;
         return true;
     }
-    
+
     return false;
 }
 
-unsigned long long FLatencyTester::EpochMillisNow()
+double FLatencyTester::EpochMillisNow()
 {
-    return std::chrono::duration_cast<std::chrono::milliseconds>(std::chrono::system_clock::now().time_since_epoch()).count();
+    return FPlatformTime::ToMilliseconds64(FPlatformTime::Cycles64());
 }
\ No newline at end of file
diff --git a/Engine/Plugins/Media/PixelStreaming/Source/PixelStreaming/Private/LatencyTester.h b/Engine/Plugins/Media/PixelStreaming/Source/PixelStreaming/Private/LatencyTester.h
--- a/Engine/Plugins/Media/PixelStreaming/Source/PixelStreaming/Private/LatencyTester.h
+++ b/Engine/Plugins/Media/PixelStreaming/Source/PixelStreaming/Private/LatencyTester.h
@@ -27,7 +27,7 @@ class FLatencyTester
         // Start latency tester, resets all the internally recorded timings, makes the RecordXXX functions do stuff (otherwise they no-op).
         static void Start(FPlayerId PlayerWhoTriggeredTest);
         static bool RecordReceiptTime();
-        static bool RecordPreCaptureTime(uint32 FrameId);
+        static bool RecordPreCaptureTime();
         static bool RecordPostCaptureTime(uint32 FrameId);
         static bool RecordPreEncodeTime(uint32 FrameId);
         static bool RecordPostEncodeTime(uint32 FrameId);
@@ -40,14 +40,14 @@ class FLatencyTester
 
         //Private constructor, only static methods are exposed.
         FLatencyTester() : TestStage(ELatencyTestStage::INACTIVE) {}
-        unsigned long long EpochMillisNow();
-
-        unsigned long long ReceiptTimeMs = 0;
-        unsigned long long PreCaptureTimeMs = 0;
-        unsigned long long PostCaptureTimeMs = 0;
-        unsigned long long PreEncodeTimeMs = 0;
-        unsigned long long PostEncodeTimeMs = 0;
-        unsigned long long TransmissionTimeMs = 0;
+        double EpochMillisNow();
+
+        double ReceiptTimeMs = 0;
+        double PreCaptureTimeMs = 0;
+        double PostCaptureTimeMs = 0;
+        double PreEncodeTimeMs = 0;
+        double PostEncodeTimeMs = 0;
+        double TransmissionTimeMs = 0;
         uint32 FrameId = 0;
         FPlayerId PlayerWhoTriggeredTest = FPlayerId();
         std::atomic<ELatencyTestStage> TestStage;
diff --git a/Engine/Plugins/Media/PixelStreaming/Source/PixelStreaming/Private/PixelStreamerInputComponent.cpp b/Engine/Plugins/Media/PixelStreaming/Source/PixelStreaming/Private/PixelStreamerInputComponent.cpp
--- a/Engine/Plugins/Media/PixelStreaming/Source/PixelStreaming/Private/PixelStreamerInputComponent.cpp
+++ b/Engine/Plugins/Media/PixelStreaming/Source/PixelStreaming/Private/PixelStreamerInputComponent.cpp
@@ -52,18 +52,18 @@ bool UPixelStreamerInputComponent::OnCommand(const FString& Descriptor)
 {
 	FString ConsoleCommand;
 	bool bSuccess = false;
-	GetJsonStringValue(Descriptor, TEXT("ConsoleCommand"), ConsoleCommand, bSuccess);
+	ExtractJsonFromDescriptor(Descriptor, TEXT("ConsoleCommand"), ConsoleCommand, bSuccess);
 	if (bSuccess)
 	{
-		return GEngine->Exec(GetWorld(), *ConsoleCommand);
+		return GEngine->Exec(GEngine->GetWorld(), *ConsoleCommand);
 	}
 	
 	FString WidthString;
 	FString HeightString;
-	GetJsonStringValue(Descriptor, TEXT("Resolution.Width"), WidthString, bSuccess);
+	ExtractJsonFromDescriptor(Descriptor, TEXT("Resolution.Width"), WidthString, bSuccess);
 	if (bSuccess)
 	{
-		GetJsonStringValue(Descriptor, TEXT("Resolution.Height"), HeightString, bSuccess);
+		ExtractJsonFromDescriptor(Descriptor, TEXT("Resolution.Height"), HeightString, bSuccess);
 
 		int Width = FCString::Atoi(*WidthString);
 		int Height = FCString::Atoi(*HeightString);
@@ -74,7 +74,7 @@ bool UPixelStreamerInputComponent::OnCommand(const FString& Descriptor)
 		}
 
 		FString ChangeResCommand = FString::Printf(TEXT("r.SetRes %dx%d"), Width, Height);
-		return GEngine->Exec(GetWorld(), *ChangeResCommand);
+		return GEngine->Exec(GEngine->GetWorld(), *ChangeResCommand);
 
 	}
 
@@ -96,6 +96,11 @@ void UPixelStreamerInputComponent::SendPixelStreamingResponse(const FString& Des
 }
 
 void UPixelStreamerInputComponent::GetJsonStringValue(FString Descriptor, FString FieldName, FString& StringValue, bool& Success)
+{
+	ExtractJsonFromDescriptor(Descriptor, FieldName, StringValue, Success);
+}
+
+void UPixelStreamerInputComponent::ExtractJsonFromDescriptor(FString Descriptor, FString FieldName, FString& StringValue, bool& Success)
 {
 	TSharedPtr<FJsonObject> JsonObject = MakeShareable(new FJsonObject);
 
@@ -129,6 +134,11 @@ void UPixelStreamerInputComponent::GetJsonStringValue(FString Descriptor, FStrin
 }
 
 void UPixelStreamerInputComponent::AddJsonStringValue(const FString& Descriptor, FString FieldName, FString StringValue, FString& NewDescriptor, bool& Success)
+{
+	ExtendJsonWithField(Descriptor, FieldName, StringValue, NewDescriptor, Success);
+}
+
+void UPixelStreamerInputComponent::ExtendJsonWithField(const FString& Descriptor, FString FieldName, FString StringValue, FString& NewDescriptor, bool& Success)
 {
 	TSharedPtr<FJsonObject> JsonObject = MakeShareable(new FJsonObject);
 
diff --git a/Engine/Plugins/Media/PixelStreaming/Source/PixelStreaming/Private/PixelStreamingEncoderFactory.cpp b/Engine/Plugins/Media/PixelStreaming/Source/PixelStreaming/Private/PixelStreamingEncoderFactory.cpp
--- a/Engine/Plugins/Media/PixelStreaming/Source/PixelStreaming/Private/PixelStreamingEncoderFactory.cpp
+++ b/Engine/Plugins/Media/PixelStreaming/Source/PixelStreaming/Private/PixelStreamingEncoderFactory.cpp
@@ -1,12 +1,12 @@
 // Copyright Epic Games, Inc. All Rights Reserved.
 
 #include "PixelStreamingEncoderFactory.h"
-#include "PixelStreamingVideoEncoder.h"
+#include "Misc/ScopeLock.h"
 #include "PixelStreamingSettings.h"
+#include "PixelStreamingVideoEncoder.h"
 #include "Utils.h"
 #include "absl/strings/match.h"
 #include "modules/video_coding/codecs/vp8/include/vp8.h"
-#include "Misc/ScopeLock.h"
 
 FPixelStreamingVideoEncoderFactory::FPixelStreamingVideoEncoderFactory(IPixelStreamingSessions* InPixelStreamingSessions)
 	: PixelStreamingSessions(InPixelStreamingSessions)
@@ -37,7 +37,7 @@ std::vector<webrtc::SdpVideoFormat> FPixelStreamingVideoEncoderFactory::GetSuppo
 
 FPixelStreamingVideoEncoderFactory::CodecInfo FPixelStreamingVideoEncoderFactory::QueryVideoEncoder(const webrtc::SdpVideoFormat& format) const
 {
-	CodecInfo codec_info = { false, false };
+	CodecInfo codec_info = {false, false};
 	codec_info.is_hardware_accelerated = true;
 	codec_info.has_internal_source = false;
 	return codec_info;
@@ -49,11 +49,7 @@ std::unique_ptr<webrtc::VideoEncoder> FPixelStreamingVideoEncoderFactory::Create
 		return webrtc::VP8Encoder::Create();
 	else
 	{
-		// Lock during encoder creation
-		FScopeLock FactoryLock(&this->FactoryCS);
 		auto VideoEncoder = std::make_unique<FPixelStreamingVideoEncoder>(this->PixelStreamingSessions, &EncoderContext);
-		this->ActiveEncoders.Add(VideoEncoder.get());
-		
 		UE_LOG(PixelStreamer, Log, TEXT("Encoder factory addded new encoder - soon to be associated with a player."));
 		return VideoEncoder;
 	}
@@ -64,18 +60,44 @@ void FPixelStreamingVideoEncoderFactory::RemoveStaleEncoders()
 	// Lock during removing stale encoders
 	FScopeLock FactoryLock(&this->FactoryCS);
 
-	// Iterate backwards so we can remove invalid encoders along the way
-	for (int32 Index = ActiveEncoders.Num()-1; Index >= 0; --Index)
+	TArray<FPlayerId> ToRemove;
+
+	for (auto& Entry : ActiveEncoders)
+	{
+		FPlayerId PlayerId = Entry.Key;
+		FPixelStreamingVideoEncoder* Encoder = Entry.Value;
+
+		// If WebRTC callback is no longer registered with the encoder it is considered stale and to be removed.
+		if (!Encoder->IsRegisteredWithWebRTC())
+		{
+			ToRemove.Add(PlayerId);
+		}
+	}
+
+	if (ToRemove.Num() > 0)
 	{
-		FPixelStreamingVideoEncoder* Encoder = ActiveEncoders[Index];
-		if(!Encoder->IsRegisteredWithWebRTC())
+		for (FPlayerId& PlayerId : ToRemove)
 		{
-			ActiveEncoders.RemoveAt(Index);
-			UE_LOG(PixelStreamer, Log, TEXT("Encoder factory cleaned up stale encoder associated with PlayerId=%s"), *Encoder->GetPlayerId());
+			ActiveEncoders.Remove(PlayerId);
+			UE_LOG(PixelStreamer, Log, TEXT("Encoder factory cleaned up stale encoder associated with PlayerId=%s"), *PlayerId);
 		}
 	}
 }
 
+void FPixelStreamingVideoEncoderFactory::OnPostEncode()
+{
+	FScopeLock FactoryLock(&this->FactoryCS);
+
+	// If we have zero encoders now then shutdown the real hardware encoder too
+	if (ActiveEncoders.Num() == 0 && EncoderContext.Encoder != nullptr)
+	{
+		UE_LOG(PixelStreamer, Log, TEXT("Encoder factory shutting down hardware encoder"));
+		EncoderContext.Encoder->ClearOnEncodedPacket();
+		EncoderContext.Encoder->Shutdown();
+		EncoderContext.Encoder = nullptr;
+	}
+}
+
 void FPixelStreamingVideoEncoderFactory::OnEncodedImage(const webrtc::EncodedImage& encoded_image, const webrtc::CodecSpecificInfo* codec_specific_info, const webrtc::RTPFragmentationHeader* fragmentation)
 {
 	// Before sending encoded image to each encoder's callback, check if all encoders we have are still relevant.
@@ -85,9 +107,10 @@ void FPixelStreamingVideoEncoderFactory::OnEncodedImage(const webrtc::EncodedIma
 	FScopeLock FactoryLock(&this->FactoryCS);
 
 	// Go through each encoder and send our encoded image to its callback
-	for(FPixelStreamingVideoEncoder* Encoder : ActiveEncoders)
+	for (auto& Entry : ActiveEncoders)
 	{
-		if(Encoder->IsRegisteredWithWebRTC())
+		FPixelStreamingVideoEncoder* Encoder = Entry.Value;
+		if (Encoder->IsRegisteredWithWebRTC())
 		{
 			Encoder->SendEncodedImage(encoded_image, codec_specific_info, fragmentation);
 		}
@@ -97,29 +120,46 @@ void FPixelStreamingVideoEncoderFactory::OnEncodedImage(const webrtc::EncodedIma
 	EncoderContext.SmoothedAvgQP.Update(encoded_image.qp_);
 }
 
-void FPixelStreamingVideoEncoderFactory::ReleaseVideoEncoder(FPixelStreamingVideoEncoder* encoder)
+void FPixelStreamingVideoEncoderFactory::RegisterVideoEncoder(FPlayerId PlayerId, FPixelStreamingVideoEncoder* Encoder)
+{
+	// Lock during adding an encoder
+	FScopeLock FactoryLock(&this->FactoryCS);
+	ActiveEncoders.Add(PlayerId, Encoder);
+}
+
+void FPixelStreamingVideoEncoderFactory::UnregisterVideoEncoder(FPlayerId PlayerId)
 {
 	// Lock during deleting an encoder
 	FScopeLock FactoryLock(&this->FactoryCS);
 
-	ActiveEncoders.Remove(encoder);
-	if (ActiveEncoders.Num() == 0 && EncoderContext.Encoder != nullptr)
+	if (ActiveEncoders.Contains(PlayerId))
 	{
-		EncoderContext.Encoder->ClearOnEncodedPacket();
-		EncoderContext.Encoder->Shutdown();
-		EncoderContext.Encoder = nullptr;
+		FPixelStreamingVideoEncoder* PixelStreamingEncoder = this->ActiveEncoders[PlayerId];
+
+		if(!PixelStreamingEncoder)
+		{
+			UE_LOG(PixelStreamer, Error, TEXT("Encoder factory tried to remove any already nullptr PixelStreamingVideoEncoder PlayerId=%s"), *PlayerId);
+			return;
+		}
+
+		// This will signal this encoder is stale and it will get removed next time we finish encoding
+		PixelStreamingEncoder->Release();
+
+		// This will ensure we don't try to send another encoded frame to this encoder.
+		ActiveEncoders.Remove(PlayerId);
+		UE_LOG(PixelStreamer, Log, TEXT("Encoder factory asked to remove encoder for PlayerId=%s"), *PlayerId);
+
 	}
-	FPlayerId PlayerId = encoder->GetPlayerId();
-	UE_LOG(PixelStreamer, Log, TEXT("Encoder factory asked to remove encoder for PlayerId=%s"), *PlayerId);
 }
 
 void FPixelStreamingVideoEncoderFactory::ForceKeyFrame()
 {
 	FScopeLock FactoryLock(&this->FactoryCS);
 	// Go through each encoder and send our encoded image to its callback
-	for(FPixelStreamingVideoEncoder* Encoder : ActiveEncoders)
+	for (auto& Entry : ActiveEncoders)
 	{
-		if(Encoder->IsRegisteredWithWebRTC())
+		FPixelStreamingVideoEncoder* Encoder = Entry.Value;
+		if (Encoder->IsRegisteredWithWebRTC())
 		{
 			Encoder->ForceKeyFrame();
 		}
diff --git a/Engine/Plugins/Media/PixelStreaming/Source/PixelStreaming/Private/PixelStreamingEncoderFactory.h b/Engine/Plugins/Media/PixelStreaming/Source/PixelStreaming/Private/PixelStreamingEncoderFactory.h
--- a/Engine/Plugins/Media/PixelStreaming/Source/PixelStreaming/Private/PixelStreamingEncoderFactory.h
+++ b/Engine/Plugins/Media/PixelStreaming/Source/PixelStreaming/Private/PixelStreamingEncoderFactory.h
@@ -8,6 +8,7 @@
 #include "PlayerId.h"
 #include "IPixelStreamingSessions.h"
 #include "HAL/CriticalSection.h"
+#include "PixelStreamingEncoderFrameFactory.h"
 
 class FPixelStreamingVideoEncoderFactory;
 class FPixelStreamingVideoEncoder;
@@ -23,9 +24,11 @@ class FPixelStreamingVideoEncoder;
 // Note: We could give an encoder per peer; however, this uses up our limited hardware encoder sessions (e.g. typically you can run only 3 encoding sessions on NVIDIA Geforce cards).
 struct FEncoderContext
 {
+	FPixelStreamingEncoderFrameFactory FrameFactory;
 	FPixelStreamingVideoEncoderFactory* Factory;
 	TUniquePtr<AVEncoder::FVideoEncoder> Encoder; // The real underlying encoder.
 	FSmoothedValue<60> SmoothedAvgQP;
+	FCriticalSection EncodingCriticalSection; // Ensures only one encoder is encoding at a time during quality controller swap.
 };
 
 class FPixelStreamingVideoEncoderFactory : public webrtc::VideoEncoderFactory
@@ -43,8 +46,10 @@ public:
 	// Always returns our H264 hardware encoder for now.
 	virtual std::unique_ptr<webrtc::VideoEncoder> CreateVideoEncoder(const webrtc::SdpVideoFormat& format) override;
 
+	void OnPostEncode();
 	void OnEncodedImage(const webrtc::EncodedImage& encoded_image, const webrtc::CodecSpecificInfo* codec_specific_info, const webrtc::RTPFragmentationHeader* fragmentation);
-	void ReleaseVideoEncoder(FPixelStreamingVideoEncoder* encoder);
+	void RegisterVideoEncoder(FPlayerId PlayerId, FPixelStreamingVideoEncoder* Encoder);
+	void UnregisterVideoEncoder(FPlayerId PlayerId);
 	void RemoveStaleEncoders();
 	void ForceKeyFrame();
 	double GetLatestQP();
@@ -53,7 +58,7 @@ private:
 	FEncoderContext EncoderContext;
 
 	// Each encoder is associated with a particular player (peer).
-	TArray<FPixelStreamingVideoEncoder*> ActiveEncoders;
+	TMap<FPlayerId, FPixelStreamingVideoEncoder*> ActiveEncoders;
 	
 	// Used for checks such as whether a given player id is associated with the quality controlling player.
 	IPixelStreamingSessions* PixelStreamingSessions;
diff --git a/Engine/Plugins/Media/PixelStreaming/Source/PixelStreaming/Private/PixelStreamingEncoderFrameFactory.cpp b/Engine/Plugins/Media/PixelStreaming/Source/PixelStreaming/Private/PixelStreamingEncoderFrameFactory.cpp
new file mode 100644
--- /dev/null
+++ b/Engine/Plugins/Media/PixelStreaming/Source/PixelStreaming/Private/PixelStreamingEncoderFrameFactory.cpp
@@ -0,0 +1,278 @@
+// Copyright Epic Games, Inc. All Rights Reserved.
+
+#include "PixelStreamingEncoderFrameFactory.h"
+#include "PixelStreamingPrivate.h"
+#include "CudaModule.h"
+#include "VulkanRHIPrivate.h"
+#include "PixelStreamingSettings.h"
+
+FPixelStreamingEncoderFrameFactory::FPixelStreamingEncoderFrameFactory()
+{
+
+}
+
+FPixelStreamingEncoderFrameFactory::~FPixelStreamingEncoderFrameFactory()
+{
+    this->FlushFrames();
+}
+
+void FPixelStreamingEncoderFrameFactory::FlushFrames()
+{
+    for(auto& Entry : this->TextureToFrameMapping)
+    {
+        this->EncoderInput->DestroyBuffer(Entry.Value);
+    }
+    this->TextureToFrameMapping.Empty();
+}
+
+AVEncoder::FVideoEncoderInputFrame* FPixelStreamingEncoderFrameFactory::GetOrCreateFrame(int InWidth, int InHeight, const FTexture2DRHIRef InTexture)
+{
+    check(EncoderInput.IsValid());
+
+    AVEncoder::FVideoEncoderInputFrame* OutFrame;
+
+    if(this->TextureToFrameMapping.Contains(InTexture))
+    {
+        OutFrame = *(this->TextureToFrameMapping.Find(InTexture));
+        checkf(InWidth == OutFrame->GetWidth() && InHeight == OutFrame->GetHeight(), TEXT("The requested resolution does not match the existing frame, we do not support a resolution change at this step."));
+    }
+    // A frame needs to be created
+    else
+    {
+        OutFrame = this->EncoderInput->CreateBuffer([this](const AVEncoder::FVideoEncoderInputFrame* ReleasedFrame){ /* OnReleased */ });
+        this->SetTexture(OutFrame, InTexture);
+        this->TextureToFrameMapping.Add(InTexture, OutFrame);
+    }
+
+    OutFrame->SetFrameID(++this->FrameId);
+    return OutFrame;
+}
+
+AVEncoder::FVideoEncoderInputFrame* FPixelStreamingEncoderFrameFactory::GetFrameAndSetTexture(int InWidth, int InHeight, FTextureObtainer TextureObtainer)
+{
+    check(EncoderInput.IsValid());
+
+    // Copy texture here into our one encoder frame
+    const FTexture2DRHIRef SourceTexture = TextureObtainer();
+
+    AVEncoder::FVideoEncoderInputFrame* Frame = this->GetOrCreateFrame(InWidth, InHeight, SourceTexture);
+
+    return Frame;
+}
+
+TSharedPtr<AVEncoder::FVideoEncoderInput> FPixelStreamingEncoderFrameFactory::GetOrCreateVideoEncoderInput(int InWidth, int InHeight)
+{
+    if(!this->EncoderInput.IsValid())
+    {
+        this->EncoderInput = CreateVideoEncoderInput(InWidth, InHeight);
+        this->EncoderInput->SetMaxNumBuffers((uint32)PixelStreamingSettings::CVarPixelStreamingMaxNumBackBuffers.GetValueOnAnyThread());
+    }
+
+    return this->EncoderInput;
+}
+
+void FPixelStreamingEncoderFrameFactory::SetResolution(int InWidth, int InHeight)
+{
+    TSharedPtr<AVEncoder::FVideoEncoderInput> NewEncoderInput = GetOrCreateVideoEncoderInput(InWidth, InHeight);
+    NewEncoderInput->SetResolution(InWidth, InHeight);
+    NewEncoderInput->Flush();
+}
+
+TSharedPtr<AVEncoder::FVideoEncoderInput> FPixelStreamingEncoderFrameFactory::CreateVideoEncoderInput(int InWidth, int InHeight) const
+{
+    if(!GDynamicRHI)
+	{
+        UE_LOG(PixelStreamer, Error, TEXT("GDynamicRHI not valid for some reason."));
+        return nullptr;
+	}
+
+    FString RHIName = GDynamicRHI->GetName();
+
+    // Consider if we want to support runtime resolution changing?
+    bool bIsResizable = false;
+
+    if(RHIName == TEXT("Vulkan"))
+    {
+        if(IsRHIDeviceAMD())
+        {
+            FVulkanDynamicRHI* DynamicRHI = static_cast<FVulkanDynamicRHI*>(GDynamicRHI);
+            AVEncoder::FVulkanDataStruct VulkanData = {DynamicRHI->GetInstance(), DynamicRHI->GetDevice()->GetPhysicalHandle(), DynamicRHI->GetDevice()->GetInstanceHandle()};
+
+            return AVEncoder::FVideoEncoderInput::CreateForVulkan(&VulkanData, InWidth, InHeight, bIsResizable);
+        }
+        else if(IsRHIDeviceNVIDIA())
+        {
+            if(FModuleManager::Get().IsModuleLoaded("CUDA"))
+            {
+                return AVEncoder::FVideoEncoderInput::CreateForCUDA(FModuleManager::GetModuleChecked<FCUDAModule>("CUDA").GetCudaContext(), InWidth, InHeight, bIsResizable);
+            }
+            else
+            {
+                UE_LOG(PixelStreamer, Error, TEXT("CUDA module is not loaded!"));
+                return nullptr;
+            }
+        }
+    }
+#if PLATFORM_WINDOWS
+    else if(RHIName == TEXT("D3D11"))
+    {
+        return AVEncoder::FVideoEncoderInput::CreateForD3D11(GDynamicRHI->RHIGetNativeDevice(), InWidth, InHeight, bIsResizable, IsRHIDeviceAMD());
+    }
+    else if(RHIName == TEXT("D3D12"))
+    {
+        return AVEncoder::FVideoEncoderInput::CreateForD3D12(GDynamicRHI->RHIGetNativeDevice(), InWidth, InHeight, bIsResizable, IsRHIDeviceNVIDIA());
+    }
+#endif
+
+    UE_LOG(PixelStreamer, Error, TEXT("Current RHI %s is not supported in Pixel Streaming"), *RHIName);
+    return nullptr;
+}
+
+void FPixelStreamingEncoderFrameFactory::SetTexture(AVEncoder::FVideoEncoderInputFrame* InputFrame, const FTexture2DRHIRef& Texture)
+{
+    FString RHIName = GDynamicRHI->GetName();
+
+    // VULKAN
+    if(RHIName == TEXT("Vulkan"))
+    {
+        if(IsRHIDeviceAMD())
+        {
+            FVulkanTexture2D* VulkanTexture = static_cast<FVulkanTexture2D*>(Texture.GetReference());
+            InputFrame->SetTexture(VulkanTexture->Surface.Image, [this, InputFrame](VkImage NativeTexture) { /* Do something with released texture if needed */ });
+        }
+        else if(IsRHIDeviceNVIDIA())
+        {
+            this->SetTextureCUDAVulkan(InputFrame, Texture);
+        }
+        else
+        {
+            UE_LOG(PixelStreamer, Error, TEXT("Pixel Streaming only supports AMD and NVIDIA devices, this device is neither of those."));
+        }
+    }
+#if PLATFORM_WINDOWS	
+    // DX11
+    else if(RHIName == TEXT("D3D11"))
+    {
+        InputFrame->SetTexture((ID3D11Texture2D*)Texture->GetNativeResource(), [this, InputFrame](ID3D11Texture2D* NativeTexture) { /* Do something with released texture if needed */ });
+    }
+    // DX12
+    else if(RHIName == TEXT("D3D12"))
+    {
+        InputFrame->SetTexture((ID3D12Resource*)Texture->GetNativeResource(), [this, InputFrame](ID3D12Resource* NativeTexture) { /* Do something with released texture if needed */ });
+    }
+#endif // PLATFORM_WINDOWS
+    else
+    {
+        UE_LOG(PixelStreamer, Error, TEXT("Pixel Streaming does not support this RHI - %s"), *RHIName);
+    }
+}
+
+void FPixelStreamingEncoderFrameFactory::SetTextureCUDAVulkan(AVEncoder::FVideoEncoderInputFrame* InputFrame, const FTexture2DRHIRef& Texture)
+{
+    FVulkanTexture2D* VulkanTexture = static_cast<FVulkanTexture2D*>(Texture.GetReference());
+    VkDevice device = static_cast<FVulkanDynamicRHI*>(GDynamicRHI)->GetDevice()->GetInstanceHandle();
+
+    // Get the CUarray to that textures memory making sure the clear it when done
+    int fd;
+
+    { 
+        // Generate VkMemoryGetFdInfoKHR
+        VkMemoryGetFdInfoKHR vkMemoryGetFdInfoKHR = {};
+        vkMemoryGetFdInfoKHR.sType = VK_STRUCTURE_TYPE_MEMORY_GET_FD_INFO_KHR;
+        vkMemoryGetFdInfoKHR.pNext = NULL;
+        vkMemoryGetFdInfoKHR.memory = VulkanTexture->Surface.GetAllocationHandle();
+        vkMemoryGetFdInfoKHR.handleType = VK_EXTERNAL_MEMORY_HANDLE_TYPE_OPAQUE_FD_BIT_KHR;
+
+        // While this operation is safe (and unavoidable) C4191 has been enabled and this will trigger an error with warnings as errors
+#pragma warning(push)
+#pragma warning(disable : 4191)
+        PFN_vkGetMemoryFdKHR fpGetMemoryFdKHR = (PFN_vkGetMemoryFdKHR)VulkanRHI::vkGetDeviceProcAddr(device, "vkGetMemoryFdKHR");
+        VERIFYVULKANRESULT(fpGetMemoryFdKHR(device, &vkMemoryGetFdInfoKHR, &fd));
+#pragma warning(pop)
+    }
+
+    FCUDAModule::CUDA().cuCtxPushCurrent(FModuleManager::GetModuleChecked<FCUDAModule>("CUDA").GetCudaContext());
+
+    CUexternalMemory mappedExternalMemory = nullptr;
+
+    {
+        // generate a cudaExternalMemoryHandleDesc
+        CUDA_EXTERNAL_MEMORY_HANDLE_DESC cudaExtMemHandleDesc = {};
+        cudaExtMemHandleDesc.type = CU_EXTERNAL_MEMORY_HANDLE_TYPE_OPAQUE_FD;
+        cudaExtMemHandleDesc.handle.fd = fd;
+        cudaExtMemHandleDesc.size = VulkanTexture->Surface.GetAllocationOffset() + VulkanTexture->Surface.GetMemorySize();
+
+        // import external memory
+        auto result = FCUDAModule::CUDA().cuImportExternalMemory(&mappedExternalMemory, &cudaExtMemHandleDesc);
+        if(result != CUDA_SUCCESS)
+        {
+            UE_LOG(PixelStreamer, Error, TEXT("Failed to import external memory from vulkan error: %d"), result);
+        }
+    }
+
+    CUmipmappedArray mappedMipArray = nullptr;
+    CUarray mappedArray = nullptr;
+
+    {
+        CUDA_EXTERNAL_MEMORY_MIPMAPPED_ARRAY_DESC mipmapDesc = {};
+        mipmapDesc.numLevels = 1;
+        mipmapDesc.offset = VulkanTexture->Surface.GetAllocationOffset();
+        mipmapDesc.arrayDesc.Width = Texture->GetSizeX();
+        mipmapDesc.arrayDesc.Height = Texture->GetSizeY();
+        mipmapDesc.arrayDesc.Depth = 0;
+        mipmapDesc.arrayDesc.NumChannels = 4;
+        mipmapDesc.arrayDesc.Format = CU_AD_FORMAT_UNSIGNED_INT8;
+        mipmapDesc.arrayDesc.Flags = CUDA_ARRAY3D_SURFACE_LDST | CUDA_ARRAY3D_COLOR_ATTACHMENT;
+
+        // get the CUarray from the external memory
+        auto result = FCUDAModule::CUDA().cuExternalMemoryGetMappedMipmappedArray(&mappedMipArray, mappedExternalMemory, &mipmapDesc);
+        if(result != CUDA_SUCCESS)
+        {
+            UE_LOG(PixelStreamer, Error, TEXT("Failed to bind mipmappedArray error: %d"), result);
+        }
+    }
+
+    // get the CUarray from the external memory
+    CUresult mipMapArrGetLevelErr = FCUDAModule::CUDA().cuMipmappedArrayGetLevel(&mappedArray, mappedMipArray, 0);
+    if(mipMapArrGetLevelErr != CUDA_SUCCESS)
+    {
+        UE_LOG(PixelStreamer, Error, TEXT("Failed to bind to mip 0."));
+    }
+
+    FCUDAModule::CUDA().cuCtxPopCurrent(NULL);
+
+    InputFrame->SetTexture(mappedArray, [this, mappedArray, mappedMipArray, mappedExternalMemory, InputFrame](CUarray NativeTexture)
+    {
+        // free the cuda types
+        FCUDAModule::CUDA().cuCtxPushCurrent(FModuleManager::GetModuleChecked<FCUDAModule>("CUDA").GetCudaContext());
+
+        if(mappedArray)
+        {
+            auto result = FCUDAModule::CUDA().cuArrayDestroy(mappedArray);
+            if(result != CUDA_SUCCESS)
+            {
+                UE_LOG(PixelStreamer, Error, TEXT("Failed to destroy mappedArray: %d"), result);
+            }
+        }
+
+        if(mappedMipArray)
+        {
+            auto result = FCUDAModule::CUDA().cuMipmappedArrayDestroy(mappedMipArray);
+            if(result != CUDA_SUCCESS)
+            {
+                UE_LOG(PixelStreamer, Error, TEXT("Failed to destroy mappedMipArray: %d"), result);
+            }
+        }
+
+        if(mappedExternalMemory)
+        {
+            auto result = FCUDAModule::CUDA().cuDestroyExternalMemory(mappedExternalMemory);
+            if(result != CUDA_SUCCESS)
+            {
+                UE_LOG(PixelStreamer, Error, TEXT("Failed to destroy mappedExternalMemoryArray: %d"), result);
+            }
+        }
+
+        FCUDAModule::CUDA().cuCtxPopCurrent(NULL);
+    });
+}
\ No newline at end of file
diff --git a/Engine/Plugins/Media/PixelStreaming/Source/PixelStreaming/Private/PixelStreamingEncoderFrameFactory.h b/Engine/Plugins/Media/PixelStreaming/Source/PixelStreaming/Private/PixelStreamingEncoderFrameFactory.h
new file mode 100644
--- /dev/null
+++ b/Engine/Plugins/Media/PixelStreaming/Source/PixelStreaming/Private/PixelStreamingEncoderFrameFactory.h
@@ -0,0 +1,36 @@
+// Copyright Epic Games, Inc. All Rights Reserved.
+
+#pragma once
+
+#include "VideoEncoderInput.h"
+#include "PixelStreamingFrameBuffer.h"
+
+/*
+* Factory to create `AVEncoder::FVideoEncoderInputFrame` for use in Pixel Streaming encoders.
+* This class is responsible for creating/managing the underlying `AVEncoder::FVideoEncoderInput` required to create
+* `AVEncoder::FVideoEncoderInputFrame`.
+* This class is not threadsafe and should only be accessed from a single thread.
+*/
+class FPixelStreamingEncoderFrameFactory
+{
+    public:
+        FPixelStreamingEncoderFrameFactory();
+        ~FPixelStreamingEncoderFrameFactory();
+        AVEncoder::FVideoEncoderInputFrame* GetFrameAndSetTexture(int InWidth, int InHeight, FTextureObtainer TextureObtainer);
+        TSharedPtr<AVEncoder::FVideoEncoderInput> GetOrCreateVideoEncoderInput(int InWidth, int InHeight);
+        void SetResolution(int InWidth, int InHeight);
+
+    private:
+        AVEncoder::FVideoEncoderInputFrame* GetOrCreateFrame(int InWidth, int InHeight, const FTexture2DRHIRef InTexture);
+        void FlushFrames();
+        TSharedPtr<AVEncoder::FVideoEncoderInput> CreateVideoEncoderInput(int InWidth, int InHeight) const;
+        void SetTexture(AVEncoder::FVideoEncoderInputFrame* Frame, const FTexture2DRHIRef& Texture);
+        void SetTextureCUDAVulkan(AVEncoder::FVideoEncoderInputFrame* Frame, const FTexture2DRHIRef& Texture);
+
+    private:
+
+        uint64 FrameId = 0;
+        TSharedPtr<AVEncoder::FVideoEncoderInput> EncoderInput;
+        // Store a mapping between raw textures and the FVideoEncoderInputFrames that wrap them
+        TMap<const FTexture2DRHIRef, AVEncoder::FVideoEncoderInputFrame*> TextureToFrameMapping;
+};
diff --git a/Engine/Plugins/Media/PixelStreaming/Source/PixelStreaming/Private/PixelStreamingFrameBuffer.h b/Engine/Plugins/Media/PixelStreaming/Source/PixelStreaming/Private/PixelStreamingFrameBuffer.h
--- a/Engine/Plugins/Media/PixelStreaming/Source/PixelStreaming/Private/PixelStreamingFrameBuffer.h
+++ b/Engine/Plugins/Media/PixelStreaming/Source/PixelStreaming/Private/PixelStreamingFrameBuffer.h
@@ -7,6 +7,8 @@
 #include "RHI.h"
 #include "PlayerId.h"
 
+using FTextureObtainer = TFunction<const FTexture2DRHIRef(void)>;
+
 /*
 * FPixelStreamingFrameBufferWrapper
 * Wraps the WebRTC VideoFrameBuffer and adds some general Pixel Streaming methods such as an encoder usage hint.
@@ -14,21 +16,19 @@
 
 class FPixelStreamingFrameBufferWrapper : public webrtc::VideoFrameBuffer
 {
-	public:
-
+public:
 	// How the encoder should use this frame buffer.
 	enum class EncoderUsageHint
 	{
 		Initialize,
 		Encode
 	};
-	
-	public:
-		virtual EncoderUsageHint GetUsageHint() const = 0;
 
- protected:
- 	virtual ~FPixelStreamingFrameBufferWrapper() override = default;
+public:
+	virtual EncoderUsageHint GetUsageHint() const = 0;
 
+protected:
+	virtual ~FPixelStreamingFrameBufferWrapper() override = default;
 };
 
 /*
@@ -41,61 +41,62 @@ class FPixelStreamingFrameBufferWrapper : public webrtc::VideoFrameBuffer
 */
 class FPixelStreamingInitFrameBuffer : public FPixelStreamingFrameBufferWrapper
 {
-    public:
-        FPixelStreamingInitFrameBuffer(FPlayerId InPlayerId, int InWidth, int InHeight)
-            : PlayerId(InPlayerId)
-            , Width(InWidth)
-            , Height(InHeight)
-        {}
-        
-        virtual ~FPixelStreamingInitFrameBuffer() = default;
-
-		EncoderUsageHint GetUsageHint() const override
-		{
-			return EncoderUsageHint::Initialize;
-		}
+public:
+	FPixelStreamingInitFrameBuffer(FPlayerId InPlayerId, int InWidth, int InHeight)
+		: PlayerId(InPlayerId)
+		, Width(InWidth)
+		, Height(InHeight)
+	{
+	}
+
+	virtual ~FPixelStreamingInitFrameBuffer() = default;
+
+	EncoderUsageHint GetUsageHint() const override
+	{
+		return EncoderUsageHint::Initialize;
+	}
+
+	// Begin webrtc::VideoFrameBuffer interface
+	virtual webrtc::VideoFrameBuffer::Type type() const
+	{
+		return webrtc::VideoFrameBuffer::Type::kNative;
+	}
+
+	virtual int width() const
+	{
+		return this->Width;
+	}
+
+	virtual int height() const
+	{
+		return this->Height;
+	}
 
-        // Begin webrtc::VideoFrameBuffer interface
-        virtual webrtc::VideoFrameBuffer::Type type() const
-        {
-            return webrtc::VideoFrameBuffer::Type::kNative;
-        }
-
-        virtual int width() const
-        {
-            return this->Width;
-        }
-
-        virtual int height() const
-        {
-            return this->Height;
-        }
-
-        virtual rtc::scoped_refptr<webrtc::I420BufferInterface> ToI420()
-        {
-            return webrtc::I420Buffer::Create(this->Width, this->Height);
-        }
-
-        virtual const webrtc::I420BufferInterface* GetI420() const
-        {
-            return nullptr;
-        }
-
-        // End webrtc::VideoFrameBuffer interface
-
-        FPlayerId GetPlayerId()
-        {
-            return this->PlayerId;
-        }
-
-	public:
-		DECLARE_DELEGATE(FOnEncoderInitialized)
-	    FOnEncoderInitialized OnEncoderInitialized;
-
-    private:
-        FPlayerId PlayerId;
-        int Width;
-        int Height;
+	virtual rtc::scoped_refptr<webrtc::I420BufferInterface> ToI420()
+	{
+		return webrtc::I420Buffer::Create(this->Width, this->Height);
+	}
+
+	virtual const webrtc::I420BufferInterface* GetI420() const
+	{
+		return nullptr;
+	}
+
+	// End webrtc::VideoFrameBuffer interface
+
+	FPlayerId GetPlayerId()
+	{
+		return this->PlayerId;
+	}
+
+public:
+	DECLARE_DELEGATE(FOnEncoderInitialized)
+	FOnEncoderInitialized OnEncoderInitialized;
+
+private:
+	FPlayerId PlayerId;
+	int Width;
+	int Height;
 };
 
 /*
@@ -111,17 +112,15 @@ class FPixelStreamingFrameBuffer : public FPixelStreamingFrameBufferWrapper
 {
 
 public:
-	explicit FPixelStreamingFrameBuffer(FTexture2DRHIRef SourceTexture, AVEncoder::FVideoEncoderInputFrame* InputFrame, TSharedPtr<AVEncoder::FVideoEncoderInput> InputVideoEncoderInput)
-		: TextureRef(SourceTexture)
-		, Frame(InputFrame)
-		, VideoEncoderInput(InputVideoEncoderInput)
+	explicit FPixelStreamingFrameBuffer(FTextureObtainer InTextureObtainer, int InWidth, int InHeight)
+		: TextureObtainer(InTextureObtainer)
+		, Width(InWidth)
+		, Height(InHeight)
 	{
-		Frame->Obtain();
 	}
 
 	~FPixelStreamingFrameBuffer()
 	{
-		Frame->Release();
 	}
 
 	EncoderUsageHint GetUsageHint() const override
@@ -129,6 +128,11 @@ public:
 		return EncoderUsageHint::Encode;
 	}
 
+	FTextureObtainer GetTextureObtainer()
+	{
+		return TextureObtainer;
+	}
+
 	//
 	// webrtc::VideoFrameBuffer interface
 	//
@@ -139,12 +143,12 @@ public:
 
 	virtual int width() const override
 	{
-		return Frame->GetWidth();
+		return Width;
 	}
 
 	virtual int height() const override
 	{
-		return Frame->GetHeight();
+		return Height;
 	}
 
 	//////////////////////////////////////////////
@@ -152,7 +156,8 @@ public:
 	//////////////////////////////////////////////
 	rtc::scoped_refptr<webrtc::I420BufferInterface> ToI420() override
 	{
-		rtc::scoped_refptr<webrtc::I420Buffer> Buffer = webrtc::I420Buffer::Create(Frame->GetWidth(), Frame->GetHeight());
+		rtc::scoped_refptr<webrtc::I420Buffer> Buffer = webrtc::I420Buffer::Create(width(), height());
+		FTexture2DRHIRef TextureRef = TextureObtainer();
 
 		// TODO Texture conversion can happen here as we can add a refernce to the Unreal texture to this class rather than doing raw RHI calls
 #if 1
@@ -164,8 +169,8 @@ public:
 			uint32 components = GPixelFormats[TextureRef->GetFormat()].NumComponents;
 
 			// Convert from BGRA to I420
-			uint32 height = Frame->GetHeight();
-			uint32 width = Frame->GetWidth();
+			uint32 height = this->height();
+			uint32 width = this->width();
 
 			uint32 image_size = width * height;
 			uint32 upos = 0;
@@ -183,7 +188,7 @@ public:
 				for (uint32 col = 0; col < TextureRef->GetSizeX(); ++col)
 				{
 					i = row * stride + col * components;
-					Pixel = { TextureData[i+2] , TextureData[i+1], TextureData[i] };
+					Pixel = {TextureData[i + 2], TextureData[i + 1], TextureData[i]};
 					if (row % 2)
 					{
 						DataY[row * width + col] = ((66 * Pixel.R + 129 * Pixel.G + 25 * Pixel.B) >> 8) + 16;
@@ -193,7 +198,7 @@ public:
 
 						col += 1;
 						i = row * stride + col * components;
-						Pixel = { TextureData[i + 2] , TextureData[i + 1], TextureData[i] };
+						Pixel = {TextureData[i + 2], TextureData[i + 1], TextureData[i]};
 						DataY[row * width + col] = ((66 * Pixel.R + 129 * Pixel.G + 25 * Pixel.B) >> 8) + 16;
 					}
 					else
@@ -214,8 +219,6 @@ public:
 		// TODO Write shader that repacks RGBA texture into I420
 		FComputeFenceRHIRef fence = GDynamicRHI->RHICreateComputeFence(TEXT("RGBA_to_YUV420_fence"));
 
-
-
 		fence.wait();
 
 		// Copy data out of I420 packed into RGBA
@@ -234,18 +237,8 @@ public:
 		return Buffer;
 	}
 
-	AVEncoder::FVideoEncoderInputFrame* GetFrame() const
-	{
-		return Frame;
-	}
-
-	TSharedPtr<AVEncoder::FVideoEncoderInput> GetVideoEncoderInput() const
-	{
-		return VideoEncoderInput;
-	}
-
 private:
-	FTexture2DRHIRef TextureRef;
-	AVEncoder::FVideoEncoderInputFrame* Frame;
-	TSharedPtr<AVEncoder::FVideoEncoderInput> VideoEncoderInput;
+	FTextureObtainer TextureObtainer;
+	int Width;
+	int Height;
 };
\ No newline at end of file
diff --git a/Engine/Plugins/Media/PixelStreaming/Source/PixelStreaming/Private/PixelStreamingModule.cpp b/Engine/Plugins/Media/PixelStreaming/Source/PixelStreaming/Private/PixelStreamingModule.cpp
--- a/Engine/Plugins/Media/PixelStreaming/Source/PixelStreaming/Private/PixelStreamingModule.cpp
+++ b/Engine/Plugins/Media/PixelStreaming/Source/PixelStreaming/Private/PixelStreamingModule.cpp
@@ -8,7 +8,6 @@
 #include "PixelStreamerDelegates.h"
 #include "SignallingServerConnection.h"
 #include "PixelStreamingSettings.h"
-#include "PixelStreamingStats.h"
 #include "PixelStreamingPrivate.h"
 #include "PlayerSession.h"
 #include "PixelStreamingAudioSink.h"
@@ -18,6 +17,7 @@
 #include "UObject/UObjectIterator.h"
 #include "Engine/Texture2D.h"
 #include "Slate/SceneViewport.h"
+#include "Utils.h"
 
 #if PLATFORM_WINDOWS || PLATFORM_XBOXONE
 #include "Windows/WindowsHWrapper.h"
@@ -78,9 +78,6 @@ namespace
 
 void FPixelStreamingModule::InitStreamer()
 {
-	// Cap the engine framerate to what WebRTC
-	GEngine->SetMaxFPS(PixelStreamingSettings::CVarPixelStreamingWebRTCMaxFps.GetValueOnAnyThread());
-
 	FString StreamerId;
 	FParse::Value(FCommandLine::Get(), TEXT("PixelStreamingID="), StreamerId);
 
@@ -307,8 +304,10 @@ void FPixelStreamingModule::FreezeFrame(UTexture2D* Texture)
 			// Create empty texture
 			FRHIResourceCreateInfo CreateInfo(TEXT("FreezeFrameTexture"));
 			FTexture2DRHIRef DestTexture = GDynamicRHI->RHICreateTexture2D(Width, Height, EPixelFormat::PF_B8G8R8A8, 1, 1, TexCreate_RenderTargetable, ERHIAccess::Present, CreateInfo);
+			FGPUFenceRHIRef CopyFence = GDynamicRHI->RHICreateGPUFence(*FString::Printf(TEXT("FreezeFrameFence")));
+
 			// Copy freeze frame texture to empty texture
-			CopyTexture(Texture2DRHI, DestTexture);
+			CopyTexture(Texture2DRHI, DestTexture, CopyFence);
 
 			TArray<FColor> Data;
 			FIntRect Rect(0, 0, Width, Height);
@@ -423,8 +422,6 @@ bool FPixelStreamingModule::IsTickableInEditor() const
 
 void FPixelStreamingModule::Tick(float DeltaTime)
 {
-	FPixelStreamingStats::Get().Tick();
-
 	// If we are running a latency test then check if we have timing results and if we do transmit them
 	if(FLatencyTester::IsTestRunning() && FLatencyTester::GetTestStage() == FLatencyTester::ELatencyTestStage::RESULTS_READY)
 	{
diff --git a/Engine/Plugins/Media/PixelStreaming/Source/PixelStreaming/Private/PixelStreamingStats.cpp b/Engine/Plugins/Media/PixelStreaming/Source/PixelStreaming/Private/PixelStreamingStats.cpp
--- a/Engine/Plugins/Media/PixelStreaming/Source/PixelStreaming/Private/PixelStreamingStats.cpp
+++ b/Engine/Plugins/Media/PixelStreaming/Source/PixelStreaming/Private/PixelStreamingStats.cpp
@@ -16,7 +16,7 @@ bool FPixelStreamingStats::GetStatsEnabled()
 	return PixelStreamingSettings::CVarPixelStreamingOnScreenStats.GetValueOnAnyThread() || PixelStreamingSettings::CVarPixelStreamingLogStats.GetValueOnAnyThread();
 }
 
-void FPixelStreamingStats::Tick()
+void FPixelStreamingStats::Tick(float DeltaTime)
 {
 	bool bStatsEnabled = this->GetStatsEnabled();
 
@@ -31,15 +31,18 @@ void FPixelStreamingStats::Tick()
 	}
 
 	double NowMillis = FTimespan::FromSeconds(FPlatformTime::Seconds()).GetTotalMilliseconds();
-	double KeyframeDeltaSecs = LastKeyFrameTimeCycles > 0 ? FPlatformTime::ToSeconds64(FPlatformTime::Cycles() - LastKeyFrameTimeCycles) : 0;
+	double KeyframeDeltaSecs = LastKeyFrameTimeCycles > 0 ? FPlatformTime::ToSeconds64(FPlatformTime::Cycles64() - LastKeyFrameTimeCycles) : 0;
 
 	int Id = 0;
 	EmitStat(++Id, FString::Printf(TEXT("Timestamp: %.0f ms"), NowMillis));
 	EmitStat(++Id, FString::Printf(TEXT("Capture latency: %.5f ms"), CaptureLatencyMs.Get()));
 	EmitStat(++Id, FString::Printf(TEXT("Encoder latency: %.5f ms"), EncoderLatencyMs.Get()));
-	EmitStat(++Id, FString::Printf(TEXT("WebRTC Capture->Encode latency: %.5f ms"), WebRTCCaptureToEncodeLatencyMs.Get()));
+	EmitStat(++Id, FString::Printf(TEXT("Post encode latency: %.5f ms"), PostEncodeLatencyMs));
+	EmitStat(++Id, FString::Printf(TEXT("WebRTC Capture to Pre-Encode latency: %.5f ms"), WebRTCCaptureToEncodeLatencyMs.Get()));
 	EmitStat(++Id, FString::Printf(TEXT("GameThread FPS: %.0f"), 1.0f / FApp::GetDeltaTime()));
 	EmitStat(++Id, FString::Printf(TEXT("Captured FPS: %.0f"), CaptureFPS.Get()));
+	EmitStat(++Id, FString::Printf(TEXT("Submit to WebRTC FPS: %.0f"), SubmitToWebRTCFPS.Get()));
+	EmitStat(++Id, FString::Printf(TEXT("WebRTC encode loop FPS: %.0f"), WebRTCEncodeLoopFPS.Get()));
 	EmitStat(++Id, FString::Printf(TEXT("Encoded FPS: %.0f"), EncoderFPS.Get()));
 	EmitStat(++Id, FString::Printf(TEXT("Encoder bitrate: %.3f Mbps"), EncoderBitrateMbps.Get()));
 	EmitStat(++Id, FString::Printf(TEXT("QP: %.0f"), EncoderQP.Get()));
@@ -63,9 +66,19 @@ void FPixelStreamingStats::EmitStat(int UniqueId, FString StringToEmit)
 
 }
 
+void FPixelStreamingStats::SetPostEncodeLatency(double InPostEncodeLatencyMs)
+{
+	this->PostEncodeLatencyMs = InPostEncodeLatencyMs;
+}
+
+void FPixelStreamingStats::SetCaptureToEncodeLatency(double CaptureToEncodeMs)
+{
+	this->WebRTCCaptureToEncodeLatencyMs.Update(CaptureToEncodeMs);
+}
+
 void FPixelStreamingStats::OnEncodingFinished()
 {
-	uint64 NowCycles = FPlatformTime::Cycles();
+	uint64 NowCycles = FPlatformTime::Cycles64();
 	if(this->LastEncodeTimeCycles != 0)
 	{
 		uint64 DeltaCycles = NowCycles - this->LastEncodeTimeCycles;
@@ -78,7 +91,7 @@ void FPixelStreamingStats::OnEncodingFinished()
 
 void FPixelStreamingStats::OnCaptureFinished()
 {
-	uint64 NowCycles = FPlatformTime::Cycles();
+	uint64 NowCycles = FPlatformTime::Cycles64();
 	if(this->LastCaptureTimeCycles != 0)
 	{
 		uint64 DeltaCycles = NowCycles - this->LastCaptureTimeCycles;
@@ -91,13 +104,28 @@ void FPixelStreamingStats::OnCaptureFinished()
 
 void FPixelStreamingStats::OnWebRTCDeliverFrameForEncode()
 {
-	uint64 NowCycles = FPlatformTime::Cycles();
-	if(this->LastCaptureTimeCycles != 0)
+	uint64 NowCycles = FPlatformTime::Cycles64();
+	if(this->LastWebRTCEncodeAttempt != 0)
 	{
-		uint64 DeltaCycles = NowCycles - this->LastCaptureTimeCycles;
-		double DeltaMs = FPlatformTime::ToMilliseconds64(DeltaCycles);
-		this->WebRTCCaptureToEncodeLatencyMs.Update(DeltaMs);
+		uint64 DeltaCycles = NowCycles - this->LastWebRTCEncodeAttempt;
+		double DeltaSeconds = FPlatformTime::ToSeconds64(DeltaCycles);
+		double LoopsPerSecond = 1.0 / DeltaSeconds;
+		this->WebRTCEncodeLoopFPS.Update(LoopsPerSecond);
+	}
+	this->LastWebRTCEncodeAttempt = NowCycles;
+}
+
+void FPixelStreamingStats::OnFrameSubmittedToWebRTC()
+{
+	uint64 NowCycles = FPlatformTime::Cycles64();
+	if(this->LastSubmitTimeCycles != 0)
+	{
+		uint64 DeltaCycles = NowCycles - this->LastSubmitTimeCycles;
+		double DeltaSeconds = FPlatformTime::ToSeconds64(DeltaCycles);
+		double SubmitsPerSecond = 1.0 / DeltaSeconds;
+		this->SubmitToWebRTCFPS.Update(SubmitsPerSecond);
 	}
+	this->LastSubmitTimeCycles = NowCycles;
 }
 
 void FPixelStreamingStats::SetEncoderLatency(double InEncoderLatencyMs)
@@ -122,7 +150,7 @@ void FPixelStreamingStats::SetCaptureLatency(double InCaptureLatencyMs)
 
 void FPixelStreamingStats::OnKeyframeEncoded()
 {
-	uint64 NowCycles = FPlatformTime::Cycles();
+	uint64 NowCycles = FPlatformTime::Cycles64();
 	this->LastKeyFrameTimeCycles = NowCycles;
 }
 
@@ -135,4 +163,6 @@ void FPixelStreamingStats::Reset()
 
 	CaptureFPS.Reset();
 	CaptureLatencyMs.Reset();
+
+	SubmitToWebRTCFPS.Reset();
 }
diff --git a/Engine/Plugins/Media/PixelStreaming/Source/PixelStreaming/Private/PixelStreamingStats.h b/Engine/Plugins/Media/PixelStreaming/Source/PixelStreaming/Private/PixelStreamingStats.h
--- a/Engine/Plugins/Media/PixelStreaming/Source/PixelStreaming/Private/PixelStreamingStats.h
+++ b/Engine/Plugins/Media/PixelStreaming/Source/PixelStreaming/Private/PixelStreamingStats.h
@@ -4,30 +4,35 @@
 
 #include "Utils.h"
 #include "PixelStreamingSettings.h"
+#include "Tickable.h"
 
 // Stats about Pixel Streaming that can displayed either in the in-application HUD or in the log.
-struct FPixelStreamingStats
+struct FPixelStreamingStats : FTickableGameObject
 {
 	public:
 		static FPixelStreamingStats& Get();
 		bool GetStatsEnabled();
-		void Tick();
+		void Tick(float DeltaTime);
 		void Reset();
 		void OnWebRTCDeliverFrameForEncode();
 		void OnEncodingFinished();
 		void OnCaptureFinished();
 		void OnKeyframeEncoded();
+		void OnFrameSubmittedToWebRTC();
+		void SetCaptureToEncodeLatency(double CaptureToEncodeMs);
 		void SetEncoderLatency(double EncoderLatencyMs);
 		void SetEncoderBitrateMbps(double EncoderBitrateMbps);
 		void SetEncoderQP(double QP);
 		void SetCaptureLatency(double CaptureLatencyMs);
+		void SetPostEncodeLatency(double PostEncodeLatencyMs);
+		FORCEINLINE TStatId GetStatId() const { RETURN_QUICK_DECLARE_CYCLE_STAT(PixelStreamingStats, STATGROUP_Tickables); }
 
 	private:
 		void EmitStat(int UniqueId, FString StringToEmit);
 
 	private:
 
-		static constexpr uint32 SmoothingPeriod = 3 * 60; // kinda 3 secs for 60FPS
+		static constexpr uint32 SmoothingPeriod = 60; // kinda 3 secs for 60FPS
 		
 		// Note: FSmoothedValue is thread safe.
 		FSmoothedValue<SmoothingPeriod> WebRTCCaptureToEncodeLatencyMs;
@@ -35,6 +40,7 @@ struct FPixelStreamingStats
 		FSmoothedValue<SmoothingPeriod> CaptureLatencyMs;
 		FSmoothedValue<SmoothingPeriod> EncoderBitrateMbps;
 		FSmoothedValue<SmoothingPeriod> EncoderQP;
+		double PostEncodeLatencyMs;
 
 		uint64 LastEncodeTimeCycles = 0;
 		FSmoothedValue<SmoothingPeriod> EncoderFPS;
@@ -42,6 +48,12 @@ struct FPixelStreamingStats
 		uint64 LastCaptureTimeCycles = 0;
 		FSmoothedValue<SmoothingPeriod> CaptureFPS;
 
+		uint64 LastSubmitTimeCycles = 0;
+		FSmoothedValue<SmoothingPeriod> SubmitToWebRTCFPS;
+
+		uint64 LastWebRTCEncodeAttempt = 0;
+		FSmoothedValue<SmoothingPeriod> WebRTCEncodeLoopFPS;
+
 		uint64 LastKeyFrameTimeCycles = 0;
 
 		
diff --git a/Engine/Plugins/Media/PixelStreaming/Source/PixelStreaming/Private/PixelStreamingVideoEncoder.cpp b/Engine/Plugins/Media/PixelStreaming/Source/PixelStreaming/Private/PixelStreamingVideoEncoder.cpp
--- a/Engine/Plugins/Media/PixelStreaming/Source/PixelStreaming/Private/PixelStreamingVideoEncoder.cpp
+++ b/Engine/Plugins/Media/PixelStreaming/Source/PixelStreaming/Private/PixelStreamingVideoEncoder.cpp
@@ -28,7 +28,7 @@ FPixelStreamingVideoEncoder::FPixelStreamingVideoEncoder(const IPixelStreamingSe
 
 FPixelStreamingVideoEncoder::~FPixelStreamingVideoEncoder()
 {
-	Context->Factory->ReleaseVideoEncoder(this);
+	Context->Factory->UnregisterVideoEncoder(this->OwnerPlayerId);
 }
 
 int FPixelStreamingVideoEncoder::InitEncode(webrtc::VideoCodec const* codec_settings, VideoEncoder::Settings const& settings)
@@ -42,11 +42,13 @@ int FPixelStreamingVideoEncoder::InitEncode(webrtc::VideoCodec const* codec_sett
 
 	EncoderConfig.Width = codec_settings->width;
 	EncoderConfig.Height = codec_settings->height;
-	EncoderConfig.MaxBitrate = codec_settings->maxBitrate;
-	EncoderConfig.TargetBitrate = codec_settings->startBitrate;
-	EncoderConfig.MaxFramerate = codec_settings->maxFramerate;
 
-	WebRtcProposedTargetBitrate = codec_settings->startBitrate;
+	// TODO investigate why codec_settings is not respecting the CVars following code is a bandaid to mitigate a 1 frame quality drop at startup due to this.
+	EncoderConfig.MaxBitrate = PixelStreamingSettings::CVarPixelStreamingWebRTCMaxBitrate.GetValueOnAnyThread(); //codec_settings->maxBitrate;
+	EncoderConfig.TargetBitrate = PixelStreamingSettings::CVarPixelStreamingWebRTCStartBitrate.GetValueOnAnyThread(); //codec_settings->startBitrate;
+	EncoderConfig.MaxFramerate = PixelStreamingSettings::CVarPixelStreamingWebRTCMaxFps.GetValueOnAnyThread(); // codec_settings->maxFramerate;
+
+	WebRtcProposedTargetBitrate = PixelStreamingSettings::CVarPixelStreamingWebRTCStartBitrate.GetValueOnAnyThread();
 
 	return WEBRTC_VIDEO_CODEC_OK;
 }
@@ -101,6 +103,10 @@ int32 FPixelStreamingVideoEncoder::Encode(webrtc::VideoFrame const& frame, std::
 	{
 		FPixelStreamingInitFrameBuffer* InitFrameBuffer = static_cast<FPixelStreamingInitFrameBuffer*>(VideoFrameBuffer);
 		this->OwnerPlayerId = InitFrameBuffer->GetPlayerId();
+
+		// Register this encoder with the factory for so it will recieve shared onEncodedCallback
+		this->Context->Factory->RegisterVideoEncoder(this->OwnerPlayerId, this);
+
 		// Let any listener know the encoder is now officially ready to receive frames for encoding.
 		InitFrameBuffer->OnEncoderInitialized.ExecuteIfBound();
 
@@ -112,6 +118,9 @@ int32 FPixelStreamingVideoEncoder::Encode(webrtc::VideoFrame const& frame, std::
 		return WEBRTC_VIDEO_CODEC_OK;
 	}
 
+	// Lock from here on in the case where quality controller is changing mid encode.
+	FScopeLock Lock(&this->Context->EncodingCriticalSection);
+
 	FPixelStreamingFrameBuffer* EncoderFrameBuffer = nullptr;
 
 	// We check if the frame is a frame we should "encode", if so, we try to proceed to actually encode it.
@@ -123,16 +132,26 @@ int32 FPixelStreamingVideoEncoder::Encode(webrtc::VideoFrame const& frame, std::
 	checkf(EncoderFrameBuffer, TEXT("Encoder framebuffer is null, encoding cannot proceed."));
 	checkf(this->OwnerPlayerId != INVALID_PLAYER_ID, TEXT("Encoder is not associated with a player/peer, encoding cannot proceed."));
 
+	if(!this->PixelStreamingSessions->IsQualityController(this->OwnerPlayerId))
+	{
+		UE_LOG(PixelStreamer, Log, TEXT("PixelStreamingVideoEncoder discarding frame - only the quality controlling peer should submit frames for encoding. This can sometimes happen on peer handover, but if you are seeing this continuously this is a bug."));
+	}
+
 	// Record stat for when WebRTC delivered frame, if stats are enabled.
 	FPixelStreamingStats& Stats = FPixelStreamingStats::Get();
 	if(Stats.GetStatsEnabled())
 	{
+		const int64 DeltaMicros = rtc::TimeMicros() - frame.timestamp_us();
+		const double DeltaMillis = (double) DeltaMicros / 1000.0;
+		Stats.SetCaptureToEncodeLatency(DeltaMillis);
 		Stats.OnWebRTCDeliverFrameForEncode();
 	}
 
+	const int FrameWidth = frame.width();
+	const int FrameHeight = frame.height();
 	if (!Context->Encoder)
 	{
-		CreateAVEncoder(EncoderFrameBuffer->GetVideoEncoderInput());
+		CreateAVEncoder(FrameWidth, FrameHeight);
 	}
 
 	// Change rates, if required.
@@ -140,13 +159,12 @@ int32 FPixelStreamingVideoEncoder::Encode(webrtc::VideoFrame const& frame, std::
 
 	// Detect video frame not matching encoder encoding resolution
 	// Note: This can happen when UE application changes its resolution and the encoder is not programattically updated.
-	int const FrameWidth = frame.width();
-	int const FrameHeight = frame.height();
 	if (EncoderConfig.Width != FrameWidth || EncoderConfig.Height != FrameHeight)
 	{
 		EncoderConfig.Width = FrameWidth;
 		EncoderConfig.Height = FrameHeight;
 		Context->Encoder->UpdateLayerConfig(0, EncoderConfig);
+		Context->FrameFactory.SetResolution(FrameWidth, FrameHeight);
 	}
 	
 	// Make a copy of encoder config from the existing config
@@ -164,8 +182,10 @@ int32 FPixelStreamingVideoEncoder::Encode(webrtc::VideoFrame const& frame, std::
 		ForceNextKeyframe = false;
 	}
 
-	AVEncoder::FVideoEncoderInputFrame* const EncoderInputFrame = EncoderFrameBuffer->GetFrame();
+	AVEncoder::FVideoEncoderInputFrame* EncoderInputFrame = Context->FrameFactory.GetFrameAndSetTexture(FrameWidth, FrameHeight, EncoderFrameBuffer->GetTextureObtainer());
 	EncoderInputFrame->SetTimestampRTP(frame.timestamp());
+	EncoderInputFrame->SetTimestampUs(frame.timestamp_us());
+
 
 	// If we are running a latency test then record pre-encode timing
 	if (FLatencyTester::IsTestRunning() && FLatencyTester::GetTestStage() == FLatencyTester::ELatencyTestStage::PRE_ENCODE)
@@ -174,8 +194,12 @@ int32 FPixelStreamingVideoEncoder::Encode(webrtc::VideoFrame const& frame, std::
 	}
 
 	// Encode the frame!
+	// TODO technically the encoder can fail here and we should return with WEBRTC_VIDEO_CODEC_ERROR 
 	Context->Encoder->Encode(EncoderInputFrame, EncodeOptions);
 
+	// Let factory know we are now post encode, so it is a safe time to do any cleanup as no one is encoding.
+	Context->Factory->OnPostEncode();
+
 	return WEBRTC_VIDEO_CODEC_OK;
 }
 
@@ -271,62 +295,18 @@ FPlayerId FPixelStreamingVideoEncoder::GetPlayerId() const
 	return this->OwnerPlayerId;
 }
 
-bool FPixelStreamingVideoEncoder::IsRegisteredWithWebRTC()
+bool FPixelStreamingVideoEncoder::IsInitialized() const
 {
-	return OnEncodedImageCallback != nullptr;
+	return this->OwnerPlayerId != INVALID_PLAYER_ID;
 }
 
-void CreateH264FragmentHeader(uint8 const* CodedData, size_t CodedDataSize, webrtc::RTPFragmentationHeader& Fragments)
+bool FPixelStreamingVideoEncoder::IsRegisteredWithWebRTC()
 {
-	// count the number of nal units
-	for (int pass = 0; pass < 2; ++pass)
-	{
-		size_t num_nal = 0;
-		size_t offset = 0;
-		while (offset < CodedDataSize)
-		{
-			// either a 0,0,1 or 0,0,0,1 sequence indicates a new 'nal'
-			size_t nal_maker_length = 3;
-			if (offset < (CodedDataSize - 3) && CodedData[offset] == 0 &&
-				CodedData[offset + 1] == 0 && CodedData[offset + 2] == 1)
-			{
-			}
-			else if (offset < (CodedDataSize - 4) && CodedData[offset] == 0 &&
-				CodedData[offset + 1] == 0 && CodedData[offset + 2] == 0 &&
-				CodedData[offset + 3] == 1)
-			{
-				nal_maker_length = 4;
-			}
-			else
-			{
-				++offset;
-				continue;
-			}
-			if (pass == 1)
-			{
-				Fragments.fragmentationOffset[num_nal] = offset + nal_maker_length;
-				Fragments.fragmentationLength[num_nal] = 0;
-				if (num_nal > 0)
-				{
-					Fragments.fragmentationLength[num_nal - 1] = offset - Fragments.fragmentationOffset[num_nal - 1];
-				}
-			}
-			offset += nal_maker_length;
-			++num_nal;
-		}
-		if (pass == 0)
-		{
-			Fragments.VerifyAndAllocateFragmentationHeader(num_nal);
-		}
-		else if (pass == 1 && num_nal > 0)
-		{
-			Fragments.fragmentationLength[num_nal - 1] = offset - Fragments.fragmentationOffset[num_nal - 1];
-		}
-	}
+	return OnEncodedImageCallback != nullptr;
 }
 
 //Note: this is a free function on purpose as it is not tied to the object life cycle of a given PixelStreamingVideoEncoder
-void OnEncodedPacket(FEncoderContext* Context, uint32 InLayerIndex, const AVEncoder::FVideoEncoderInputFrame* InFrame, const AVEncoder::FCodecPacketImpl& InPacket)
+void OnEncodedPacket(FEncoderContext* Context, uint32 InLayerIndex, const AVEncoder::FVideoEncoderInputFrame* InFrame, const AVEncoder::FCodecPacket& InPacket, uint64 StartPostEncodeCycles)
 {
 	// During shutdown this can function can sometimes be called by a queued encode if timing is unfortunate and the context is null
 	if(Context == nullptr)
@@ -336,8 +316,17 @@ void OnEncodedPacket(FEncoderContext* Context, uint32 InLayerIndex, const AVEnco
 
 	webrtc::EncodedImage Image;
 
+	// Create H264 frag header for the encoded frame
 	webrtc::RTPFragmentationHeader FragHeader;
-	CreateH264FragmentHeader(InPacket.Data, InPacket.DataSize, FragHeader);
+	std::vector<webrtc::H264::NaluIndex> NALUIndices = webrtc::H264::FindNaluIndices(InPacket.Data, InPacket.DataSize);
+	FragHeader.VerifyAndAllocateFragmentationHeader(NALUIndices.size());
+	FragHeader.fragmentationVectorSize = static_cast<uint16_t>(NALUIndices.size());
+	for (int I = 0; I != NALUIndices.size(); ++I)
+	{
+		webrtc::H264::NaluIndex const& NALUIndex = NALUIndices[I];
+		FragHeader.fragmentationOffset[I] = NALUIndex.payload_start_offset;
+		FragHeader.fragmentationLength[I] = NALUIndex.payload_size;
+	}
 
 	Image.timing_.packetization_finish_ms = FTimespan::FromSeconds(FPlatformTime::Seconds()).GetTotalMilliseconds();
 	Image.timing_.encode_start_ms = InPacket.Timings.StartTs.GetTotalMilliseconds();
@@ -371,6 +360,11 @@ void OnEncodedPacket(FEncoderContext* Context, uint32 InLayerIndex, const AVEnco
 
 	if (Stats.GetStatsEnabled())
 	{
+		uint64 EndPostEncodeCycles = FPlatformTime::Cycles64();
+		double PostEncodeMs = FPlatformTime::ToMilliseconds64(EndPostEncodeCycles - StartPostEncodeCycles);
+		Stats.SetPostEncodeLatency(PostEncodeMs);
+		//UE_LOG(PixelStreamer, Log, TEXT("Post encode ms %f"), PostEncodeMs);
+
 		Stats.SetEncoderLatency(EncoderLatencyMs);
 		Stats.SetEncoderBitrateMbps(BitrateMbps);
 		Stats.SetEncoderQP(InPacket.VideoQP);
@@ -388,35 +382,19 @@ void OnEncodedPacket(FEncoderContext* Context, uint32 InLayerIndex, const AVEnco
 	}
 }
 
-void FPixelStreamingVideoEncoder::CreateAVEncoder(TSharedPtr<AVEncoder::FVideoEncoderInput> EncoderInput)
+void FPixelStreamingVideoEncoder::CreateAVEncoder(int Width, int Height)
 {
 	// TODO: When we have multiple HW encoders do some factory checking and find the best encoder.
 	const TArray<AVEncoder::FVideoEncoderInfo>& Available = AVEncoder::FVideoEncoderFactory::Get().GetAvailable();
 
-	Context->Encoder = AVEncoder::FVideoEncoderFactory::Get().Create(Available[0].ID, EncoderInput, EncoderConfig);
+	Context->Encoder = AVEncoder::FVideoEncoderFactory::Get().Create(Available[0].ID, Context->FrameFactory.GetOrCreateVideoEncoderInput(Width, Height), EncoderConfig);
 	checkf(Context->Encoder, TEXT("Pixel Streaming video encoder creation failed, check encoder config."));
 
 	FEncoderContext* ContextPtr = this->Context;
 	Context->Encoder->SetOnEncodedPacket([ContextPtr](uint32 InLayerIndex, const AVEncoder::FVideoEncoderInputFrame* InFrame, const AVEncoder::FCodecPacket& InPacket) 
 	{
-
-		// Create a memory copy of the CodecPacket because the encoder will recycle InPacket.Data again after this.
-		AVEncoder::FCodecPacketImpl* CopyPacket = new AVEncoder::FCodecPacketImpl();
-		CopyPacket->DataSize = InPacket.DataSize;
-		CopyPacket->Data = static_cast<const uint8*>(FMemory::Malloc(InPacket.DataSize));
-		FMemory::BigBlockMemcpy(const_cast<uint8*>(CopyPacket->Data), InPacket.Data, InPacket.DataSize);
-		CopyPacket->IsKeyFrame = InPacket.IsKeyFrame;
-		CopyPacket->VideoQP = InPacket.VideoQP;
-		CopyPacket->Timings = InPacket.Timings;
-
-		// We do the actual work somewhere on the task graph so we are not locking up the encoder.
-		AsyncTask(ENamedThreads::AnyHiPriThreadHiPriTask, [ContextPtr, InLayerIndex, InFrame, CopyPacket] ()
-		{
-			OnEncodedPacket(ContextPtr, InLayerIndex, InFrame, *CopyPacket);
-			delete[] CopyPacket->Data;
-			delete CopyPacket;
-		}); 
-		
+		uint64 StartPostEncodeCycles = FPlatformTime::Cycles64();
+		OnEncodedPacket(ContextPtr, InLayerIndex, InFrame, InPacket, StartPostEncodeCycles);
 	});
 }
 
diff --git a/Engine/Plugins/Media/PixelStreaming/Source/PixelStreaming/Private/PixelStreamingVideoEncoder.h b/Engine/Plugins/Media/PixelStreaming/Source/PixelStreaming/Private/PixelStreamingVideoEncoder.h
--- a/Engine/Plugins/Media/PixelStreaming/Source/PixelStreaming/Private/PixelStreamingVideoEncoder.h
+++ b/Engine/Plugins/Media/PixelStreaming/Source/PixelStreaming/Private/PixelStreamingVideoEncoder.h
@@ -38,16 +38,16 @@ public:
 
 	void SendEncodedImage(webrtc::EncodedImage const& encoded_image, webrtc::CodecSpecificInfo const* codec_specific_info, webrtc::RTPFragmentationHeader const* fragmentation);
 	FPlayerId GetPlayerId() const;
+	bool IsInitialized() const;
 	bool IsRegisteredWithWebRTC();
 
 	void ForceKeyFrame() { ForceNextKeyframe = true; }
 	int32_t GetSmoothedAverageQP() const;
-	
 
 private:
 	void UpdateConfig(AVEncoder::FVideoEncoder::FLayerConfig const& Config);
 	void HandlePendingRateChange();
-	void CreateAVEncoder(TSharedPtr<AVEncoder::FVideoEncoderInput> EncoderInput);
+	void CreateAVEncoder(int Width, int Height);
 	AVEncoder::FVideoEncoder::FLayerConfig CreateEncoderConfigFromCVars(AVEncoder::FVideoEncoder::FLayerConfig BaseEncoderConfig) const;
 
 	// We store this so we can restore back to it if the user decides to use then stop using the PixelStreaming.Encoder.TargetBitrate CVar.
@@ -63,7 +63,8 @@ private:
 	// The alternative is encoding separate streams for each peer, which is not tenable while NVENC sessions are limited.
 	FPlayerId OwnerPlayerId;
 
-	bool ForceNextKeyframe = false;
+	// This flag lets us force a keyframe through manually, it also lets us ensure the first encoded frame is a keyframe.
+	bool ForceNextKeyframe = true;
 
 	// USed for checks such as whether a given player id is associated with the quality controlling player.
 	const IPixelStreamingSessions* PixelStreamingSessions;
diff --git a/Engine/Plugins/Media/PixelStreaming/Source/PixelStreaming/Private/PixelStreamingVideoSources.cpp b/Engine/Plugins/Media/PixelStreaming/Source/PixelStreaming/Private/PixelStreamingVideoSources.cpp
--- a/Engine/Plugins/Media/PixelStreaming/Source/PixelStreaming/Private/PixelStreamingVideoSources.cpp
+++ b/Engine/Plugins/Media/PixelStreaming/Source/PixelStreaming/Private/PixelStreamingVideoSources.cpp
@@ -6,30 +6,46 @@
 #include "PixelStreamingSettings.h"
 #include "Async/Async.h"
 
-#define SUBMIT_RENDERING_TASK_WITH_PARAMS(Func, ...) \
-    if(IsInRenderingThread()) \
-    { \
-        this->Func(__VA_ARGS__); \
-    } \
-    else \
-    { \
-        AsyncTask(ENamedThreads::ActualRenderingThread, [this, __VA_ARGS__](){ this->Func(__VA_ARGS__); } ); \
-    } \
-
 FPixelStreamingVideoSources::FPixelStreamingVideoSources()
     : bMatchFrameBufferResolution(true)
-    , bFixedResolution(PixelStreamingSettings::CVarPixelStreamingWebRTCDisableResolutionChange.GetValueOnAnyThread()){}
+    , bFixedResolution(PixelStreamingSettings::CVarPixelStreamingWebRTCDisableResolutionChange.GetValueOnAnyThread())
+    , ThreadWaiter(FPlatformProcess::GetSynchEventFromPool(false))
+    , bIsThreadRunning(true)
+{
+	this->FrameDelayMs = (1.0f / PixelStreamingSettings::CVarPixelStreamingWebRTCMaxFps.GetValueOnAnyThread()) * 1000.0f;
+
+    this->VideoSourcesThread = MakeUnique<FThread>(TEXT("VideoSourcesThread"), [this]() { this->RunVideoSourcesLoop_VideoSourcesThread(); });
+}
 
 FPixelStreamingVideoSources::FPixelStreamingVideoSources(int Width, int Height)
     : StartResolution(FIntPoint(Width, Height))
     , bMatchFrameBufferResolution(false)
-    , bFixedResolution(true){}
+    , bFixedResolution(true)
+    , ThreadWaiter(FPlatformProcess::GetSynchEventFromPool(false))
+    , bIsThreadRunning(true)
+{
+    this->VideoSourcesThread = MakeUnique<FThread>(TEXT("VideoSourcesThread"), [this]() { this->RunVideoSourcesLoop_VideoSourcesThread(); });
+}
+
+FPixelStreamingVideoSources::~FPixelStreamingVideoSources()
+{
+    this->bIsThreadRunning = false;
+    this->ThreadWaiter->Trigger();
+    this->VideoSourcesThread->Join();
+}
+
+void FPixelStreamingVideoSources::QueueTask(TFunction<void()> InTaskToQueue)
+{
+    this->QueuedTasks.Enqueue(InTaskToQueue);
+    // This will wake the VideoSourcesThread if it is sleeping.
+    this->ThreadWaiter->Trigger();
+}
 
 webrtc::VideoTrackSourceInterface* FPixelStreamingVideoSources::CreateVideoSource(FPlayerId PlayerId)
 {
-    if(IsInRenderingThread())
+    if(IsInVideoSourcesThread())
     {
-        return this->CreateVideoSource_RenderThread(PlayerId);
+        return this->CreateVideoSource_VideoSourcesThread(PlayerId);
     }
     else
     {
@@ -39,15 +55,15 @@ webrtc::VideoTrackSourceInterface* FPixelStreamingVideoSources::CreateVideoSourc
         {
             FScopedEvent ThreadBlocker;
 
-            AsyncTask(ENamedThreads::ActualRenderingThread, [this, PlayerId, &VideoSource, &ThreadBlocker]()
+            this->QueueTask([this, PlayerId, &VideoSource, &ThreadBlocker]()
             {
-                VideoSource = this->CreateVideoSource_RenderThread(PlayerId);
+                VideoSource = this->CreateVideoSource_VideoSourcesThread(PlayerId);
 
                 // Unblocks the thread
                 ThreadBlocker.Trigger();
             });
 
-            // Blocks current thread until render thread calls trigger (indicating it is done)
+            // Blocks current thread until video sources thread calls trigger (indicating it is done)
         }
 
         return VideoSource;
@@ -56,61 +72,132 @@ webrtc::VideoTrackSourceInterface* FPixelStreamingVideoSources::CreateVideoSourc
 
 void FPixelStreamingVideoSources::SetQualityController(FPlayerId PlayerId)
 {
-    SUBMIT_RENDERING_TASK_WITH_PARAMS(SetQualityController_RenderThread, PlayerId)
+    this->QueueTask([this, PlayerId](){ this->SetQualityController_VideoSourcesThread(PlayerId); });
 }
 
 void FPixelStreamingVideoSources::DeleteVideoSource(FPlayerId PlayerId)
 {
-    SUBMIT_RENDERING_TASK_WITH_PARAMS(DeleteVideoSource_RenderThread, PlayerId)
+    this->QueueTask([this, PlayerId](){ this->DeleteVideoSource_VideoSourcesThread(PlayerId); });
 }
 
 void FPixelStreamingVideoSources::OnFrameReady(const FTexture2DRHIRef& FrameBuffer)
 {
     checkf(IsInRenderingThread(), TEXT("This method must be called on the rendering thread."));
 
-    // There is no controlling player, so discard the incoming frame as this implies there are no video sources.
-    if(this->ControllingPlayer == INVALID_PLAYER_ID)
-    {
-        return;
-    }
-
-    // If we have not created video encoder input yet, do it here.
+    // If we have not created video capturer yet, do it here.
     if(!this->VideoCapturerContext.IsValid())
     {
         // Make FVideoEncoderInput at the appropriate resolution based on VideoSource settings
         this->StartResolution = this->bMatchFrameBufferResolution ? FrameBuffer->GetSizeXY() : this->StartResolution;
-        this->VideoCapturerContext = MakeShared<FVideoCapturerContext>(this->StartResolution.X, this->StartResolution.Y, this->bFixedResolution);
+        this->VideoCapturerContext = MakeShared<FVideoCapturerContext, ESPMode::ThreadSafe>(this->StartResolution.X, this->StartResolution.Y, this->bFixedResolution);
+    }
+
+    // Copy frame.
+    this->VideoCapturerContext->CaptureFrame(FrameBuffer);
+
+    // Backbuffer now has something captured, thread can wake up (if needed).
+    this->ThreadWaiter->Trigger();
+}
+
+bool FPixelStreamingVideoSources::IsInVideoSourcesThread()
+{
+    if(this->VideoSourcesThread == nullptr)
+    {
+        return false;
+    }
+
+    uint32 VideoSourcesThreadId = this->VideoSourcesThread->GetThreadId();
+    uint32 CurrentThreadId = FPlatformTLS::GetCurrentThreadId();
+
+    return VideoSourcesThreadId == CurrentThreadId;
+}
+
+////////////////////////////////////////////
+// INTERNAL VIDEO SOURCES THREAD FUNCTIONS
+////////////////////////////////////////////
+
+void FPixelStreamingVideoSources::RunVideoSourcesLoop_VideoSourcesThread()
+{
+    uint64 LastFrameSubmittedCycles = 0;
+
+    while(this->bIsThreadRunning)
+    {
+        // Process any queued work
+        while(!this->QueuedTasks.IsEmpty())
+        {
+            TFunction<void()>* Task = this->QueuedTasks.Peek();
+            if(Task != nullptr && *Task)
+            {
+                (*Task)();
+            }
+            this->QueuedTasks.Pop();
+        }
+
+        // Submit Frame if we have a video source
+        if(this->VideoSources.Num() > 0 && this->VideoCapturerContext.IsValid() && this->VideoCapturerContext->IsInitialized())
+        {
+
+            this->SubmitFrame_VideoSourcesThread();
+
+            // Determine how long we should sleep before submitting a frame again
+            uint64 DeltaCycles = FPlatformTime::Cycles64() - LastFrameSubmittedCycles;
+            double DeltaMs = FPlatformTime::ToMilliseconds64(DeltaCycles);
+            LastFrameSubmittedCycles = FPlatformTime::Cycles64();
+
+            if(DeltaMs < this->FrameDelayMs)
+            {
+                // Sleep for DeltaMs or until ThreadWaiter.Trigger() is called.
+                double SleepTimeMs = this->FrameDelayMs - DeltaMs;
+                this->ThreadWaiter->Wait(SleepTimeMs, false);
+            }
+        }
+        // No video sources, we will sleep until we have a video source or some queued work to do
+        else
+        {
+            // Sleep indefinitely until ThreadWaiter.Trigger() is called somewhere.
+            if(this->QueuedTasks.IsEmpty())
+            {
+                this->ThreadWaiter->Wait();
+            }
+        }
     }
 
+}
+
+void FPixelStreamingVideoSources::SubmitFrame_VideoSourcesThread()
+{
+    checkf(IsInVideoSourcesThread(), TEXT("This method must called on the VideoSourcesThread"));
+
+    // Send the frame from the video source back into WebRTC where it will eventually makes its way into an encoder
     // Iterate each video source and initialize it, if needed.
     for(auto& Entry : this->VideoSources)
     {
         TUniquePtr<FVideoCapturer>& Capturer = Entry.Value;
         if(!Capturer->IsInitialized())
         {
-            Capturer->Initialize(FrameBuffer, this->VideoCapturerContext);
+            FIntPoint StartRes(this->VideoCapturerContext->GetCaptureWidth(), this->VideoCapturerContext->GetCaptureHeight());
+            Capturer->Initialize(StartRes);
         }
     }
 
-    TUniquePtr<FVideoCapturer>& ControllingVideoSource = this->VideoSources.FindChecked(this->ControllingPlayer);
-
-    checkf(ControllingVideoSource.IsValid(), TEXT("Video source was invalid, this indicates a bug and should never happen."));
+    TUniquePtr<FVideoCapturer>* ControllingVideoSourcePtr = this->VideoSources.Find(this->ControllingPlayer);
 
-    // Pass along the frame buffer to the controlling video source
-    if(ControllingVideoSource->IsInitialized())
+    // If there is no controlling player associated with any video source we should not transmit
+    if(ControllingVideoSourcePtr == nullptr || (*ControllingVideoSourcePtr).IsValid() == false)
     {
-        ControllingVideoSource->OnFrameReady(FrameBuffer);
+        return;
     }
 
+    // Pass along the capturer context to the controlling video source
+    if((*ControllingVideoSourcePtr)->IsInitialized())
+    {
+        (*ControllingVideoSourcePtr)->TrySubmitFrame(this->VideoCapturerContext);
+    }
 }
 
-///////////////////////
-// INTERNAL
-///////////////////////
-
-void FPixelStreamingVideoSources::SetQualityController_RenderThread(FPlayerId PlayerId)
+void FPixelStreamingVideoSources::SetQualityController_VideoSourcesThread(FPlayerId PlayerId)
 {
-    checkf(IsInRenderingThread(), TEXT("This method must be called on the rendering thread."));
+    checkf(IsInVideoSourcesThread(), TEXT("This method must called on the VideoSourcesThread"));
 
     if(this->VideoSources.Contains(PlayerId))
     {
@@ -122,9 +209,9 @@ void FPixelStreamingVideoSources::SetQualityController_RenderThread(FPlayerId Pl
     }
 }
 
-webrtc::VideoTrackSourceInterface* FPixelStreamingVideoSources::CreateVideoSource_RenderThread(FPlayerId PlayerId)
+webrtc::VideoTrackSourceInterface* FPixelStreamingVideoSources::CreateVideoSource_VideoSourcesThread(FPlayerId PlayerId)
 {
-    checkf(IsInRenderingThread(), TEXT("This method must be called on the rendering thread."));
+    checkf(IsInVideoSourcesThread(), TEXT("This method must called on the VideoSourcesThread"));
 
     bool bShouldMakeController = this->VideoSources.Num() == 0;
 
@@ -141,15 +228,15 @@ webrtc::VideoTrackSourceInterface* FPixelStreamingVideoSources::CreateVideoSourc
 
     if(bShouldMakeController)
     {
-        this->SetQualityController_RenderThread(PlayerId);
+        this->SetQualityController_VideoSourcesThread(PlayerId);
     }
 
     return WebRTCVideoSource;
 }
 
-void FPixelStreamingVideoSources::DeleteVideoSource_RenderThread(FPlayerId PlayerId)
+void FPixelStreamingVideoSources::DeleteVideoSource_VideoSourcesThread(FPlayerId PlayerId)
 {
-    checkf(IsInRenderingThread(), TEXT("This method must be called on the rendering thread."));
+    checkf(IsInVideoSourcesThread(), TEXT("This method must called on the VideoSourcesThread"));
 
     this->VideoSources.Remove(PlayerId);
 
@@ -158,4 +245,4 @@ void FPixelStreamingVideoSources::DeleteVideoSource_RenderThread(FPlayerId Playe
     {
         this->ControllingPlayer = INVALID_PLAYER_ID;
     }
-}
\ No newline at end of file
+}
diff --git a/Engine/Plugins/Media/PixelStreaming/Source/PixelStreaming/Private/PixelStreamingVideoSources.h b/Engine/Plugins/Media/PixelStreaming/Source/PixelStreaming/Private/PixelStreamingVideoSources.h
--- a/Engine/Plugins/Media/PixelStreaming/Source/PixelStreaming/Private/PixelStreamingVideoSources.h
+++ b/Engine/Plugins/Media/PixelStreaming/Source/PixelStreaming/Private/PixelStreamingVideoSources.h
@@ -8,15 +8,12 @@
 #include "Templates/UniquePtr.h"
 #include "WebRTCIncludes.h"
 #include "Math/IntPoint.h"
+#include "HAL/Thread.h"
 
 // All possible video sources that Pixel Streaming could send a framebuffer to.
 // Currently this class is designed so only one video source is actively sent frames - however, this design may change is we support multiple frame buffers.
 //
-// NOTE: This class is threadsafe as all methods execute on the same thread - the rendering thread!
-// The reason for executing work on the rendering thread is that in order to pass the framebuffer (via OnFrameReady) to WebRTC as fast as possible 
-// we don't want to queue that work onto some other thread to run at a later time.
-// Additionally, as we never want to lock the rendering thread when underlying video sources changes,
-// our solution is to change underlying resources on render thread.
+// NOTE: This class is threadsafe as all public methods execute on the same thread - the "video sources thread".
 class FPixelStreamingVideoSources
 {
 
@@ -28,24 +25,37 @@ class FPixelStreamingVideoSources
         // Specify the resolution of the video source, note if this constructor is used we assume resolution cannot be changed for this source now.
         FPixelStreamingVideoSources(int Width, int Height);
 
+        ~FPixelStreamingVideoSources();
+
         void OnFrameReady(const FTexture2DRHIRef& FrameBuffer);
         webrtc::VideoTrackSourceInterface* CreateVideoSource(FPlayerId PlayerId);
         void DeleteVideoSource(FPlayerId PlayerId);
 
         // /*Quality Controller*/
         void SetQualityController(FPlayerId PlayerId);
-        // FPlayerId GetQualityController();
+
+    /*Functions that run exclusively on the VideoSourcesThread*/
+    private:
+        webrtc::VideoTrackSourceInterface* CreateVideoSource_VideoSourcesThread(FPlayerId PlayerId);
+        void DeleteVideoSource_VideoSourcesThread(FPlayerId PlayerId);
+        void SetQualityController_VideoSourcesThread(FPlayerId PlayerId);
+        void RunVideoSourcesLoop_VideoSourcesThread();
+        void SubmitFrame_VideoSourcesThread();
 
     private:
-        webrtc::VideoTrackSourceInterface* CreateVideoSource_RenderThread(FPlayerId PlayerId);
-        void DeleteVideoSource_RenderThread(FPlayerId PlayerId);
-        void SetQualityController_RenderThread(FPlayerId PlayerId);
+        bool IsInVideoSourcesThread();
+        void QueueTask(TFunction<void()> InTaskToQueue);
 
     private:
         FPlayerId ControllingPlayer = INVALID_PLAYER_ID;
         TMap<FPlayerId, TUniquePtr<FVideoCapturer>> VideoSources;
-        TSharedPtr<FVideoCapturerContext> VideoCapturerContext = nullptr;
+        TSharedPtr<FVideoCapturerContext, ESPMode::ThreadSafe> VideoCapturerContext = nullptr;
         FIntPoint StartResolution;
         bool bMatchFrameBufferResolution;
         bool bFixedResolution;
+        TUniquePtr<FThread> VideoSourcesThread;
+        FEvent* ThreadWaiter;
+        TQueue<TFunction<void()>> QueuedTasks;
+        FThreadSafeBool bIsThreadRunning;
+		float FrameDelayMs = (1.0 / 60.0f) * 1000.0f;
 };
\ No newline at end of file
diff --git a/Engine/Plugins/Media/PixelStreaming/Source/PixelStreaming/Private/Streamer.cpp b/Engine/Plugins/Media/PixelStreaming/Source/PixelStreaming/Private/Streamer.cpp
--- a/Engine/Plugins/Media/PixelStreaming/Source/PixelStreaming/Private/Streamer.cpp
+++ b/Engine/Plugins/Media/PixelStreaming/Source/PixelStreaming/Private/Streamer.cpp
@@ -230,6 +230,7 @@ void FStreamer::HandleOffer(FPlayerId PlayerId, webrtc::PeerConnectionInterface*
 					codec_params.max_bitrate_bps = PixelStreamingSettings::CVarPixelStreamingWebRTCMaxBitrate.GetValueOnAnyThread();
 					codec_params.min_bitrate_bps = PixelStreamingSettings::CVarPixelStreamingWebRTCMinBitrate.GetValueOnAnyThread();
 					codec_params.max_framerate = PixelStreamingSettings::CVarPixelStreamingWebRTCMaxFps.GetValueOnAnyThread();
+					codec_params.network_priority = webrtc::Priority::kHigh;
 				}
 				
 				// Set the degradation preference based on our CVar for it.
@@ -416,12 +417,17 @@ void FStreamer::DeleteAllPlayerSessions()
 
 void FStreamer::DeletePlayerSession(FPlayerId PlayerId)
 {
+	// Note: this is very explicitly the first thing to get called, other ordering is problematic at the time of writing.
+	this->VideoEncoderFactory->UnregisterVideoEncoder(PlayerId);
+
 	int NumRemainingPlayers = this->PlayerSessions.DeletePlayerSession(PlayerId);
 
 	if(NumRemainingPlayers == 0)
 	{
 		this->bStreamingStarted = false;
 	}
+
+	this->VideoSources.DeleteVideoSource(PlayerId);
 }
 
 void FStreamer::AddStreams(FPlayerId PlayerId, webrtc::PeerConnectionInterface* PeerConnection)
diff --git a/Engine/Plugins/Media/PixelStreaming/Source/PixelStreaming/Private/ThreadSafePlayerSessions.cpp b/Engine/Plugins/Media/PixelStreaming/Source/PixelStreaming/Private/ThreadSafePlayerSessions.cpp
--- a/Engine/Plugins/Media/PixelStreaming/Source/PixelStreaming/Private/ThreadSafePlayerSessions.cpp
+++ b/Engine/Plugins/Media/PixelStreaming/Source/PixelStreaming/Private/ThreadSafePlayerSessions.cpp
@@ -413,15 +413,15 @@ int FThreadSafePlayerSessions::DeletePlayerSession_SignallingThread(FPlayerId Pl
     Players.Remove(PlayerId);
 	delete Player;
 
+    // Player deleted, tell all our C++ listeners.
+    this->OnPlayerDeleted.Broadcast(PlayerId);
+
 	UPixelStreamerDelegates* Delegates = UPixelStreamerDelegates::GetPixelStreamerDelegates();
 	if (Delegates)
 	{
 		Delegates->OnClosedConnection.Broadcast(PlayerId, bWasQualityController);
 	}
 
-    // Player deleted, tell all our C++ listeners.
-    this->OnPlayerDeleted.Broadcast(PlayerId);
-
 	// this is called from WebRTC signalling thread, the only thread were `Players` map is modified, so no need to lock it
 	if (Players.Num() == 0)
 	{
@@ -530,4 +530,9 @@ void FThreadSafePlayerSessions::SetQualityController_SignallingThread(FPlayerId
     {
         UE_LOG(PixelStreamer, Log, TEXT("Could not set quality controller for PlayerId=%s - that player does not exist."), *PlayerId);
     }
-}
\ No newline at end of file
+}
+
+#undef SUBMIT_TASK_WITH_RETURN
+#undef SUBMIT_TASK_WITH_PARAMS_AND_RETURN
+#undef SUBMIT_TASK_NO_PARAMS
+#undef SUBMIT_TASK_WITH_PARAMS
\ No newline at end of file
diff --git a/Engine/Plugins/Media/PixelStreaming/Source/PixelStreaming/Private/Utils.h b/Engine/Plugins/Media/PixelStreaming/Source/PixelStreaming/Private/Utils.h
--- a/Engine/Plugins/Media/PixelStreaming/Source/PixelStreaming/Private/Utils.h
+++ b/Engine/Plugins/Media/PixelStreaming/Source/PixelStreaming/Private/Utils.h
@@ -177,7 +177,7 @@ private:
 };
 
 
-inline void CopyTexture(const FTexture2DRHIRef& SourceTexture, FTexture2DRHIRef& DestinationTexture)
+inline void CopyTexture(const FTexture2DRHIRef& SourceTexture, FTexture2DRHIRef& DestinationTexture, FGPUFenceRHIRef& CopyFence)
 {
 	FRHICommandListImmediate& RHICmdList = FRHICommandListExecutor::GetImmediateCommandList();
 
@@ -220,15 +220,18 @@ inline void CopyTexture(const FTexture2DRHIRef& SourceTexture, FTexture2DRHIRef&
 		}
 
 		RendererModule->DrawRectangle(RHICmdList, 0, 0,                // Dest X, Y
-		                              DestinationTexture->GetSizeX(),  // Dest Width
-		                              DestinationTexture->GetSizeY(),  // Dest Height
-		                              0, 0,                            // Source U, V
-		                              1, 1,                            // Source USize, VSize
-		                              DestinationTexture->GetSizeXY(), // Target buffer size
-		                              FIntPoint(1, 1),                 // Source texture size
-		                              VertexShader, EDRF_Default);
+									DestinationTexture->GetSizeX(),  // Dest Width
+									DestinationTexture->GetSizeY(),  // Dest Height
+									0, 0,                            // Source U, V
+									1, 1,                            // Source USize, VSize
+									DestinationTexture->GetSizeXY(), // Target buffer size
+									FIntPoint(1, 1),                 // Source texture size
+									VertexShader, EDRF_Default);
 	}
 
 	RHICmdList.EndRenderPass();
-}
 
+	RHICmdList.WriteGPUFence(CopyFence);
+
+	RHICmdList.ImmediateFlush(EImmediateFlushType::WaitForOutstandingTasksOnly);
+}
diff --git a/Engine/Plugins/Media/PixelStreaming/Source/PixelStreaming/Private/VideoCapturer.cpp b/Engine/Plugins/Media/PixelStreaming/Source/PixelStreaming/Private/VideoCapturer.cpp
--- a/Engine/Plugins/Media/PixelStreaming/Source/PixelStreaming/Private/VideoCapturer.cpp
+++ b/Engine/Plugins/Media/PixelStreaming/Source/PixelStreaming/Private/VideoCapturer.cpp
@@ -8,11 +8,10 @@
 #include "HAL/PlatformTime.h"
 #include "Misc/Timespan.h"
 #include "Modules/ModuleManager.h"
-#include "PixelStreamingFrameBuffer.h"
 #include "PixelStreamingSettings.h"
+#include "PixelStreamingFrameBuffer.h"
 #include "PixelStreamingStats.h"
 #include "LatencyTester.h"
-#include "VideoCapturerContext.h"
 
 FVideoCapturer::FVideoCapturer(FPlayerId InPlayerId)
 	: bInitialized(false)
@@ -25,6 +24,7 @@ FVideoCapturer::FVideoCapturer(FPlayerId InPlayerId)
 FVideoCapturer::~FVideoCapturer()
 {
 	CurrentState = webrtc::MediaSourceInterface::SourceState::kEnded;
+	UE_LOG(PixelStreamer, Log, TEXT("Destroying WebRTC video track source associated with player: %s."), *PlayerId);
 }
 
 bool FVideoCapturer::IsInitialized()
@@ -32,7 +32,7 @@ bool FVideoCapturer::IsInitialized()
 	return this->bInitialized;
 }
 
-void FVideoCapturer::Initialize(const FTexture2DRHIRef& FrameBuffer, TSharedPtr<FVideoCapturerContext> InCapturerContext)
+void FVideoCapturer::Initialize(FIntPoint& StartResolution)
 {
 	// Check if already initialized
 	if(this->bInitialized)
@@ -40,20 +40,17 @@ void FVideoCapturer::Initialize(const FTexture2DRHIRef& FrameBuffer, TSharedPtr<
 		return;
 	}
 
-	this->CapturerContext = InCapturerContext;
-
-	const int64 TimestampUs = rtc::TimeMicros();
-	if(!AdaptCaptureFrame(TimestampUs, FrameBuffer->GetSizeXY()))
+	if(CurrentState != webrtc::MediaSourceInterface::SourceState::kLive)
 	{
-		return;
+		CurrentState = webrtc::MediaSourceInterface::SourceState::kLive;
 	}
 
 	// Here we pass the frame to Webrtc which will pass it to the correct encoder.
 	// We pass a special frame called the FPixelStreamingInitFrameBuffer that contains the player id, so the encoder is associated with a specific player id.
 	rtc::scoped_refptr<FPixelStreamingInitFrameBuffer> Buffer = new rtc::RefCountedObject<FPixelStreamingInitFrameBuffer>(
 		PlayerId, 
-		this->CapturerContext->GetCaptureWidth(), 
-		this->CapturerContext->GetCaptureHeight());
+		StartResolution.X, 
+		StartResolution.Y);
 
 	// Bind to when encoder is actually intialized as we will use this as a signal to start sending proper frames for encoding
 	Buffer->OnEncoderInitialized.BindRaw(this, &FVideoCapturer::OnEncoderInitialized);
@@ -61,7 +58,7 @@ void FVideoCapturer::Initialize(const FTexture2DRHIRef& FrameBuffer, TSharedPtr<
 	// Build a WebRTC frame (note frame id hardcoded to zero, we assume this is okay)
 	webrtc::VideoFrame Frame = webrtc::VideoFrame::Builder()
 		.set_video_frame_buffer(Buffer)
-		.set_timestamp_us(TimestampUs)
+		.set_timestamp_us(rtc::TimeMicros())
 		.set_rotation(webrtc::VideoRotation::kVideoRotation_0)
 		.set_id(0).build();
 
@@ -69,52 +66,19 @@ void FVideoCapturer::Initialize(const FTexture2DRHIRef& FrameBuffer, TSharedPtr<
 	OnFrame(Frame);
 }
 
-void FVideoCapturer::OnEncoderInitialized()
+bool FVideoCapturer::TrySubmitFrame(TSharedPtr<FVideoCapturerContext, ESPMode::ThreadSafe> CapturerContext)
 {
-	this->bInitialized = true;
-}
-
-void FVideoCapturer::OnFrameReady(const FTexture2DRHIRef& FrameBuffer)
-{
-	
-	if(!this->bInitialized)
-	{
-		// Not initialized, can't do anything with this framebuffer
-		return;
-	}
+	FPixelStreamingStats& Stats = FPixelStreamingStats::Get();
 
 	const int64 TimestampUs = rtc::TimeMicros();
-	if(!AdaptCaptureFrame(TimestampUs, FrameBuffer->GetSizeXY()))
-	{
-		return;
-	}
-
-	if(CurrentState != webrtc::MediaSourceInterface::SourceState::kLive)
-	{
-		CurrentState = webrtc::MediaSourceInterface::SourceState::kLive;
-	}
-	
-	FVideoCapturerContext::FCapturerInput CapturerInput = this->CapturerContext->ObtainCapturerInput();
-
-	if(CapturerInput.InputFrame == nullptr || !CapturerInput.Texture.IsSet())
-	{
-		return;
-	}
-
-	AVEncoder::FVideoEncoderInputFrame* InputFrame = CapturerInput.InputFrame;
-	FTexture2DRHIRef Texture = CapturerInput.Texture.GetValue();
 
-	const int32 FrameId = InputFrame->GetFrameID();
-	InputFrame->SetTimestampUs(TimestampUs);
-
-	// Latency test pre capture
-	if(FLatencyTester::IsTestRunning() && FLatencyTester::GetTestStage() == FLatencyTester::ELatencyTestStage::PRE_CAPTURE)
+	if(!AdaptCaptureFrame(TimestampUs, CapturerContext))
 	{
-		FLatencyTester::RecordPreCaptureTime(FrameId);
+		return false;
 	}
 
-	// Actual texture copy (i.e the actual "capture")
-	CopyTexture(FrameBuffer, Texture);
+	// Generate a new ID for this frame
+	int32 FrameId = CapturerContext->GetNextFrameId();
 
 	// Latency test post capture
 	if(FLatencyTester::IsTestRunning() && FLatencyTester::GetTestStage() == FLatencyTester::ELatencyTestStage::POST_CAPTURE)
@@ -122,70 +86,73 @@ void FVideoCapturer::OnFrameReady(const FTexture2DRHIRef& FrameBuffer)
 		FLatencyTester::RecordPostCaptureTime(FrameId);
 	}
 
-	UE_LOG(PixelStreamer, VeryVerbose, TEXT("(%d) captured video %lld"), RtcTimeMs(), TimestampUs);
+	FTextureObtainer TextureObtainer = CapturerContext->RequestNewestCapturedFrame();
 
-	// Here we pass the frame to Webrtc which will pass it to the correct encoder.
-	rtc::scoped_refptr<FPixelStreamingFrameBuffer> Buffer = new rtc::RefCountedObject<FPixelStreamingFrameBuffer>(Texture, InputFrame, this->CapturerContext->GetVideoEncoderInput());
-	webrtc::VideoFrame Frame = webrtc::VideoFrame::Builder()
-		.set_video_frame_buffer(Buffer)
-		.set_timestamp_us(TimestampUs)
-		.set_rotation(webrtc::VideoRotation::kVideoRotation_0)
-		.set_id(FrameId)
-		.build();
-	
-	// Send the frame from the video source back into WebRTC where it will eventually makes its way into an encoder
-	OnFrame(Frame);
+    // Here we pass the frame to Webrtc which will pass it to the correct encoder.
+    rtc::scoped_refptr<FPixelStreamingFrameBuffer> Buffer = new rtc::RefCountedObject<FPixelStreamingFrameBuffer>(TextureObtainer, CapturerContext->GetCaptureWidth(), CapturerContext->GetCaptureHeight());
+    webrtc::VideoFrame WebRTCFrame = webrtc::VideoFrame::Builder()
+        .set_video_frame_buffer(Buffer)
+        .set_timestamp_us(TimestampUs)
+        .set_rotation(webrtc::VideoRotation::kVideoRotation_0)
+        .set_id(FrameId)
+        .build();
 
-	InputFrame->Release();
+	OnFrame(WebRTCFrame);
 
-	// If stats are enabled, records the stats during capture now.
-	FPixelStreamingStats& Stats = FPixelStreamingStats::Get();
-	if(Stats.GetStatsEnabled())
+	// Record stat for how often we are submitting to WebRTC
+	if (Stats.GetStatsEnabled())
 	{
-		int64 TimestampNowUs = rtc::TimeMicros();
-		int64 CaptureLatencyUs = TimestampNowUs - TimestampUs;
-		double CaptureLatencyMs = (double)CaptureLatencyUs / 1000.0;
-		Stats.SetCaptureLatency(CaptureLatencyMs);
-		Stats.OnCaptureFinished();
+		Stats.OnFrameSubmittedToWebRTC();
 	}
-
+	
+	return true;
 }
 
+void FVideoCapturer::OnEncoderInitialized()
+{
+	this->bInitialized = true;
+}
 
-
-bool FVideoCapturer::AdaptCaptureFrame(const int64 TimestampUs, FIntPoint Resolution)
+bool FVideoCapturer::AdaptCaptureFrame(const int64 TimestampUs, TSharedPtr<FVideoCapturerContext, ESPMode::ThreadSafe> CapturerContext)
 {
+	bool bWebRTCWantsTheFrame = true;
+
 	int outWidth, outHeight, cropWidth, cropHeight, cropX, cropY;
-	if(!AdaptFrame(Resolution.X, Resolution.Y, TimestampUs, &outWidth, &outHeight, &cropWidth, &cropHeight, &cropX, &cropY))
+	if(!AdaptFrame(CapturerContext->GetCaptureWidth(), CapturerContext->GetCaptureHeight(), TimestampUs, &outWidth, &outHeight, &cropWidth, &cropHeight, &cropX, &cropY))
 	{
-		return false;
+		bWebRTCWantsTheFrame = false;
 	}
 
-	// Set resolution of encoder using user-defined params (i.e. not the back buffer).
-	if(!PixelStreamingSettings::CVarPixelStreamingUseBackBufferCaptureSize.GetValueOnRenderThread())
+	FIntPoint PreferredCaptureRes;
+
+	// Set Resolution of encoder using user-defined params (i.e. not the back buffer).
+	if(!PixelStreamingSettings::CVarPixelStreamingUseBackBufferCaptureSize.GetValueOnAnyThread())
 	{
-		// set resolution based on cvars
-		FString CaptureSize = PixelStreamingSettings::CVarPixelStreamingCaptureSize.GetValueOnRenderThread();
+		// set Resolution based on cvars
+		FString CaptureSize = PixelStreamingSettings::CVarPixelStreamingCaptureSize.GetValueOnAnyThread();
 		FString TargetWidth, TargetHeight;
 		bool bValidSize = CaptureSize.Split(TEXT("x"), &TargetWidth, &TargetHeight);
 		if(bValidSize)
 		{
-			Resolution.X = FCString::Atoi(*TargetWidth);
-			Resolution.Y = FCString::Atoi(*TargetHeight);
+			PreferredCaptureRes.X = FCString::Atoi(*TargetWidth);
+			PreferredCaptureRes.Y = FCString::Atoi(*TargetHeight);
 		}
 		else
 		{
 			UE_LOG(PixelStreamer, Error, TEXT("CVarPixelStreamingCaptureSize is not in a valid format: %s. It should be e.g: \"1920x1080\""), *CaptureSize);
-			PixelStreamingSettings::CVarPixelStreamingCaptureSize->Set(*FString::Printf(TEXT("%dx%d"), Resolution.X, Resolution.Y));
+			PixelStreamingSettings::CVarPixelStreamingCaptureSize->Set(*FString::Printf(TEXT("%dx%d"), PreferredCaptureRes.X, PreferredCaptureRes.Y));
 		}
 	}
 	else
 	{
-		Resolution.X = outWidth;
-		Resolution.Y = outHeight;
+		PreferredCaptureRes.X = outWidth;
+		PreferredCaptureRes.Y = outHeight;
 	}
 
-	this->CapturerContext->SetCaptureResolution(Resolution.X, Resolution.Y);
+	if(PreferredCaptureRes.X != CapturerContext->GetCaptureWidth() || PreferredCaptureRes.Y != CapturerContext->GetCaptureHeight())
+	{
+		CapturerContext->SetCaptureResolution(PreferredCaptureRes.X, PreferredCaptureRes.Y);
+	}
 
-	return true;
+	return bWebRTCWantsTheFrame;
 }
diff --git a/Engine/Plugins/Media/PixelStreaming/Source/PixelStreaming/Private/VideoCapturer.h b/Engine/Plugins/Media/PixelStreaming/Source/PixelStreaming/Private/VideoCapturer.h
--- a/Engine/Plugins/Media/PixelStreaming/Source/PixelStreaming/Private/VideoCapturer.h
+++ b/Engine/Plugins/Media/PixelStreaming/Source/PixelStreaming/Private/VideoCapturer.h
@@ -4,7 +4,6 @@
 
 #include "PixelStreamingPrivate.h"
 #include "WebRTCIncludes.h"
-#include "RHI.h"
 #include "PlayerId.h"
 #include "VideoCapturerContext.h"
 
@@ -17,8 +16,8 @@ public:
 	~FVideoCapturer();
 
 	bool IsInitialized();
-	void Initialize(const FTexture2DRHIRef& FrameBuffer, TSharedPtr<FVideoCapturerContext> CapturerContext);
-	void OnFrameReady(const FTexture2DRHIRef& FrameBuffer);
+	void Initialize(FIntPoint& StartResolution);
+	bool TrySubmitFrame(TSharedPtr<FVideoCapturerContext, ESPMode::ThreadSafe> CapturerContext);
 
 	void AddRef() const override
 	{
@@ -47,7 +46,7 @@ public:
 
 	virtual bool is_screencast() const override
 	{
-		return false;
+		return true;
 	}
 
 	virtual absl::optional<bool> needs_denoising() const override
@@ -56,13 +55,12 @@ public:
 	}
 
 private:
-	bool AdaptCaptureFrame(const int64 TimestampUs, FIntPoint Resolution);
+	bool AdaptCaptureFrame(const int64 TimestampUs, TSharedPtr<FVideoCapturerContext, ESPMode::ThreadSafe> CapturerContext);
 	void SetCaptureResolution(int width, int height);
 	void OnEncoderInitialized();
 
 	bool bInitialized;
 	FPlayerId PlayerId;
-	TSharedPtr<FVideoCapturerContext> CapturerContext;
 	webrtc::MediaSourceInterface::SourceState CurrentState;
 	volatile int32 Count;
 };
diff --git a/Engine/Plugins/Media/PixelStreaming/Source/PixelStreaming/Private/VideoCapturerContext.cpp b/Engine/Plugins/Media/PixelStreaming/Source/PixelStreaming/Private/VideoCapturerContext.cpp
--- a/Engine/Plugins/Media/PixelStreaming/Source/PixelStreaming/Private/VideoCapturerContext.cpp
+++ b/Engine/Plugins/Media/PixelStreaming/Source/PixelStreaming/Private/VideoCapturerContext.cpp
@@ -6,20 +6,25 @@
 #include "VulkanRHIPrivate.h"
 #include "Misc/CoreDelegates.h"
 #include "PixelStreamingPrivate.h"
+#include "LatencyTester.h"
+#include "PixelStreamingStats.h"
+#include "RHIResources.h"
+#include "Async/Async.h"
+#include "GenericPlatform/GenericPlatformProcess.h"
 
 FVideoCapturerContext::FVideoCapturerContext(int InCaptureWidth, int InCaptureHeight, bool bInFixedResolution)
     : CaptureWidth(InCaptureWidth)
     , CaptureHeight(InCaptureHeight)
     , bFixedResolution(bInFixedResolution)
-    , VideoEncoderInput(CreateVideoEncoderInput(InCaptureWidth, InCaptureHeight, bInFixedResolution))
+    , bIsInitialized(false)
 {
-    this->VideoEncoderInput->SetMaxNumBuffers((uint32)PixelStreamingSettings::CVarPixelStreamingMaxNumBackBuffers.GetValueOnRenderThread());
-    FCoreDelegates::OnPreExit.AddRaw(this, &FVideoCapturerContext::DeleteBackBuffers);
-}
+	this->EvenFrame.CopyFence = GDynamicRHI->RHICreateGPUFence(TEXT("VideoCapturerCopyFenceEven"));
+	this->OddFrame.CopyFence = GDynamicRHI->RHICreateGPUFence(TEXT("VideoCapturerCopyFenceOdd"));
 
-void FVideoCapturerContext::DeleteBackBuffers()
-{
-    BackBuffers.Empty();
+	EvenFrame.Texture = MakeTexture();
+	OddFrame.Texture = MakeTexture();
+	TempTexture = MakeTexture();
+	EncoderTexture = MakeTexture();
 }
 
 bool FVideoCapturerContext::IsFixedResolution() const
@@ -37,9 +42,26 @@ int FVideoCapturerContext::GetCaptureHeight() const
     return this->CaptureHeight;
 }
 
-TSharedPtr<AVEncoder::FVideoEncoderInput> FVideoCapturerContext::GetVideoEncoderInput() const
+int32 FVideoCapturerContext::GetNextFrameId()
 {
-    return this->VideoEncoderInput;
+    // todo: make this return uint16 and wrap when we reach max
+    return this->NextFrameID.Increment();
+}
+
+bool FVideoCapturerContext::IsInitialized() const
+{
+    return this->bIsInitialized;
+}
+
+FTexture2DRHIRef FVideoCapturerContext::MakeTexture()
+{
+	FRHIResourceCreateInfo CreateInfo(TEXT("VideoCapturerBackBuffer"));
+    FString RHIName = GDynamicRHI->GetName();
+
+	// TODO check that these are being created correctly and that they are transitioned for copies
+	return (RHIName == TEXT("Vulkan")) ?
+			GDynamicRHI->RHICreateTexture2D(this->CaptureWidth, this->CaptureHeight, EPixelFormat::PF_B8G8R8A8, 1, 1, TexCreate_RenderTargetable | TexCreate_UAV, ERHIAccess::Present, CreateInfo) :
+			GDynamicRHI->RHICreateTexture2D(this->CaptureWidth, this->CaptureHeight, EPixelFormat::PF_B8G8R8A8, 1, 1, TexCreate_Shared | TexCreate_RenderTargetable, ERHIAccess::CopyDest, CreateInfo);
 }
 
 void FVideoCapturerContext::SetCaptureResolution(int NewCaptureWidth, int NewCaptureHeight)
@@ -61,279 +83,75 @@ void FVideoCapturerContext::SetCaptureResolution(int NewCaptureWidth, int NewCap
 
 	this->CaptureWidth = NewCaptureWidth;
 	this->CaptureHeight = NewCaptureHeight;
-
-	this->VideoEncoderInput->SetResolution(this->CaptureWidth, this->CaptureHeight);
-
-    // Flushes available frames only, active frames still get a chance to go through the pipeline and be naturally removed from the backbuffers.
-	this->VideoEncoderInput->Flush();
 }
 
-TSharedPtr<AVEncoder::FVideoEncoderInput> FVideoCapturerContext::CreateVideoEncoderInput(int InWidth, int InHeight, bool bInFixedResolution)
+void FVideoCapturerContext::CaptureFrame(const FTexture2DRHIRef& FrameBuffer)
 {
-    if(!GDynamicRHI)
+	// Latency test pre capture
+	if(FLatencyTester::IsTestRunning() && FLatencyTester::GetTestStage() == FLatencyTester::ELatencyTestStage::PRE_CAPTURE)
 	{
-        UE_LOG(PixelStreamer, Error, TEXT("GDynamicRHI not valid for some reason."));
-        return nullptr;
+		FLatencyTester::RecordPreCaptureTime();
 	}
 
-    FString RHIName = GDynamicRHI->GetName();
-    bool bIsResizable = !bInFixedResolution;
-
-    if(RHIName == TEXT("Vulkan"))
-    {
-        if(IsRHIDeviceAMD())
-        {
-            FVulkanDynamicRHI* DynamicRHI = static_cast<FVulkanDynamicRHI*>(GDynamicRHI);
-            AVEncoder::FVulkanDataStruct VulkanData = {DynamicRHI->GetInstance(), DynamicRHI->GetDevice()->GetPhysicalHandle(), DynamicRHI->GetDevice()->GetInstanceHandle()};
-
-            return AVEncoder::FVideoEncoderInput::CreateForVulkan(&VulkanData, InWidth, InHeight, bIsResizable);
-        }
-        else if(IsRHIDeviceNVIDIA())
-        {
-            if(FModuleManager::Get().IsModuleLoaded("CUDA"))
-            {
-                return AVEncoder::FVideoEncoderInput::CreateForCUDA(FModuleManager::GetModuleChecked<FCUDAModule>("CUDA").GetCudaContext(), InWidth, InHeight, bIsResizable);
-            }
-            else
-            {
-                UE_LOG(PixelStreamer, Error, TEXT("CUDA module is not loaded!"));
-                return nullptr;
-            }
-        }
-    }
-#if PLATFORM_WINDOWS
-    else if(RHIName == TEXT("D3D11"))
-    {
-        return AVEncoder::FVideoEncoderInput::CreateForD3D11(GDynamicRHI->RHIGetNativeDevice(), InWidth, InHeight, bIsResizable, IsRHIDeviceAMD());
-    }
-    else if(RHIName == TEXT("D3D12"))
-    {
-        return AVEncoder::FVideoEncoderInput::CreateForD3D12(GDynamicRHI->RHIGetNativeDevice(), InWidth, InHeight, bIsResizable, IsRHIDeviceNVIDIA());
-    }
-#endif
+	bIsEvenFrame = !bIsEvenFrame;
+	FVideoCaptureFrame& Frame = bIsEvenFrame ? EvenFrame : OddFrame;
 
-    UE_LOG(PixelStreamer, Error, TEXT("Current RHI %s is not supported in Pixel Streaming"), *RHIName);
-    return nullptr;
-
-}
-
-#if PLATFORM_WINDOWS
-FTexture2DRHIRef FVideoCapturerContext::SetBackbufferTextureDX11(AVEncoder::FVideoEncoderInputFrame* InputFrame)
-{
-    FRHIResourceCreateInfo CreateInfo(TEXT("VideoCapturerBackBuffer"));
-    FTexture2DRHIRef Texture = GDynamicRHI->RHICreateTexture2D(this->CaptureWidth, this->CaptureHeight, EPixelFormat::PF_B8G8R8A8, 1, 1, TexCreate_Shared | TexCreate_RenderTargetable, ERHIAccess::CopyDest, CreateInfo);
-    InputFrame->SetTexture((ID3D11Texture2D*)Texture->GetNativeResource(), [this, InputFrame](ID3D11Texture2D* NativeTexture) { this->BackBuffers.Remove(InputFrame); });
-    this->BackBuffers.Add(InputFrame, Texture);
-    return Texture;
-}
+	Frame.CopyFence->Clear();
 
-FTexture2DRHIRef FVideoCapturerContext::SetBackbufferTextureDX12(AVEncoder::FVideoEncoderInputFrame* InputFrame)
-{
-    FRHIResourceCreateInfo CreateInfo(TEXT("VideoCapturerBackBuffer"));
-    FTexture2DRHIRef Texture = GDynamicRHI->RHICreateTexture2D(this->CaptureWidth, this->CaptureHeight, EPixelFormat::PF_B8G8R8A8, 1, 1, TexCreate_Shared | TexCreate_RenderTargetable, ERHIAccess::CopyDest, CreateInfo);
-    InputFrame->SetTexture((ID3D12Resource*)Texture->GetNativeResource(), [this, InputFrame](ID3D12Resource* NativeTexture) { this->BackBuffers.Remove(InputFrame); });
-    this->BackBuffers.Add(InputFrame, Texture);
-    return Texture;
-}
-#endif // PLATFORM_WINDOWS
-
-FTexture2DRHIRef FVideoCapturerContext::SetBackbufferTexturePureVulkan(AVEncoder::FVideoEncoderInputFrame* InputFrame)
-{
-    FRHIResourceCreateInfo CreateInfo(TEXT("VideoCapturerBackBuffer"));
-    FTexture2DRHIRef Texture =
-        GDynamicRHI->RHICreateTexture2D(this->CaptureWidth, this->CaptureHeight, EPixelFormat::PF_B8G8R8A8, 1, 1, TexCreate_Shared | TexCreate_RenderTargetable | TexCreate_UAV, ERHIAccess::Present, CreateInfo);
-
-    FVulkanTexture2D* VulkanTexture = static_cast<FVulkanTexture2D*>(Texture.GetReference());
-    InputFrame->SetTexture(VulkanTexture->Surface.Image, [this, InputFrame](VkImage NativeTexture) { this->BackBuffers.Remove(InputFrame); });
-    this->BackBuffers.Add(InputFrame, Texture);
-    return Texture;
-}
+	FTexture2DRHIRef& DestinationTexture = Frame.Texture;
+ 
+	FRHICommandListImmediate& RHICmdList = FRHICommandListExecutor::GetImmediateCommandList();
 
-FVideoCapturerContext::FCapturerInput FVideoCapturerContext::ObtainCapturerInput()
-{
-    if(!this->VideoEncoderInput.IsValid())
-    {
-        UE_LOG(PixelStreamer, Error, TEXT("VideoEncoderInput is nullptr cannot capture a frame."));
-        return FVideoCapturerContext::FCapturerInput();
-    }
+	RHICmdList.EnqueueLambda([&](FRHICommandListImmediate& RHICmdList){
+		Frame.PreWaitingOnCopy = FPlatformTime::Cycles64();
+	});
 
-    // Obtain a frame from video encoder input, we use this frame to store an RHI specific texture.
-    // Note: obtain frame will recycle frames when they are no longer being used and become "available".
-	AVEncoder::FVideoEncoderInputFrame* InputFrame = VideoEncoderInput->ObtainInputFrame();
+	// Create a bunch of RHI commands to start a "copy" (render pass) on the render thread
+	CopyTexture(FrameBuffer, DestinationTexture, Frame.CopyFence);
 
-	if(InputFrame == nullptr)
+	AsyncTask(ENamedThreads::AnyHiPriThreadHiPriTask, [&]()
 	{
-		return FVideoCapturerContext::FCapturerInput();
-	}
-
-    // Back buffer already contains a texture for this particular frame, no need to go and make one.
-    if(BackBuffers.Contains(InputFrame))
-    {
-        return FVideoCapturerContext::FCapturerInput(InputFrame, this->BackBuffers[InputFrame]);
-    }
-
-    // Got here, backbuffer does not contain this frame/texture already, so we must create a new platform specific texture.
-    FString RHIName = GDynamicRHI->GetName();
-    
-    FTexture2DRHIRef OutTexture;
-
-    // VULKAN
-    if(RHIName == TEXT("Vulkan"))
-    {
-        if(IsRHIDeviceAMD())
-        {
-            OutTexture = this->SetBackbufferTexturePureVulkan(InputFrame);
-        }
-        else if(IsRHIDeviceNVIDIA())
-        {
-            OutTexture = this->SetBackbufferTextureCUDAVulkan(InputFrame);
-        }
-        else
-        {
-            UE_LOG(PixelStreamer, Error, TEXT("Pixel Streaming only supports AMD and NVIDIA devices, this device is neither of those."));
-            return FVideoCapturerContext::FCapturerInput();
-        }
-    }
-#if PLATFORM_WINDOWS	
-    // DX11
-    else if(RHIName == TEXT("D3D11"))
-    {
-        OutTexture = this->SetBackbufferTextureDX11(InputFrame);
-    }
-    // DX12
-    else if(RHIName == TEXT("D3D12"))
-    {
-        OutTexture = this->SetBackbufferTextureDX12(InputFrame);
-    }
-#endif // PLATFORM_WINDOWS
-    else
-    {
-        UE_LOG(PixelStreamer, Error, TEXT("Pixel Streaming does not support this RHI - %s"), *RHIName);
-        return FVideoCapturerContext::FCapturerInput();
-    }
-
-    return FVideoCapturerContext::FCapturerInput(InputFrame, OutTexture);
+		// Sleep this thread until the CopyFence is completed.
+		FPlatformProcess::ConditionalSleep([Frame]()-> bool{ return Frame.CopyFence->Poll();}, TNumericLimits<float>::Max());
+		
+		{
+			FScopeLock Lock(&CriticalSection);
+			
+			TempTexture.Swap(Frame.Texture);
+			Frame.CopyFence->Clear();
+
+			bIsTempDirty = true;
+		}
+
+		uint64 PostWaitingOnCopy = FPlatformTime::Cycles64();
+
+		FPixelStreamingStats& Stats = FPixelStreamingStats::Get();
+		if(Stats.GetStatsEnabled())
+		{
+			double CaptureLatencyMs = FPlatformTime::ToMilliseconds64(PostWaitingOnCopy - Frame.PreWaitingOnCopy);
+			Stats.SetCaptureLatency(CaptureLatencyMs);
+			Stats.OnCaptureFinished();
+		}
+	});
+
+	bIsInitialized = true;
 }
 
-FTexture2DRHIRef FVideoCapturerContext::SetBackbufferTextureCUDAVulkan(AVEncoder::FVideoEncoderInputFrame* InputFrame)
+FTextureObtainer FVideoCapturerContext::RequestNewestCapturedFrame()
 {
-    FRHIResourceCreateInfo CreateInfo(TEXT("VideoCapturerBackBuffer"));
-    FVulkanDynamicRHI* VulkanDynamicRHI = static_cast<FVulkanDynamicRHI*>(GDynamicRHI);
-
-    // Create a texture that can be exposed to external memory
-    FTexture2DRHIRef Texture =
-        GDynamicRHI->RHICreateTexture2D(this->CaptureWidth, this->CaptureHeight, EPixelFormat::PF_B8G8R8A8, 1, 1, TexCreate_Shared | TexCreate_RenderTargetable | TexCreate_UAV, ERHIAccess::Present, CreateInfo);
-
-    FVulkanTexture2D* VulkanTexture = static_cast<FVulkanTexture2D*>(Texture.GetReference());
-    VkDevice device = static_cast<FVulkanDynamicRHI*>(GDynamicRHI)->GetDevice()->GetInstanceHandle();
-
-    // Get the CUarray to that textures memory making sure the clear it when done
-    int fd;
-
-    { 
-        // Generate VkMemoryGetFdInfoKHR
-        VkMemoryGetFdInfoKHR vkMemoryGetFdInfoKHR = {};
-        vkMemoryGetFdInfoKHR.sType = VK_STRUCTURE_TYPE_MEMORY_GET_FD_INFO_KHR;
-        vkMemoryGetFdInfoKHR.pNext = NULL;
-        vkMemoryGetFdInfoKHR.memory = VulkanTexture->Surface.GetAllocationHandle();
-        vkMemoryGetFdInfoKHR.handleType = VK_EXTERNAL_MEMORY_HANDLE_TYPE_OPAQUE_FD_BIT_KHR;
-
-        // While this operation is safe (and unavoidable) C4191 has been enabled and this will trigger an error with warnings as errors
-#pragma warning(push)
-#pragma warning(disable : 4191)
-        PFN_vkGetMemoryFdKHR fpGetMemoryFdKHR = (PFN_vkGetMemoryFdKHR)VulkanRHI::vkGetDeviceProcAddr(device, "vkGetMemoryFdKHR");
-        VERIFYVULKANRESULT(fpGetMemoryFdKHR(device, &vkMemoryGetFdInfoKHR, &fd));
-#pragma warning(pop)
-    }
-
-    FCUDAModule::CUDA().cuCtxPushCurrent(FModuleManager::GetModuleChecked<FCUDAModule>("CUDA").GetCudaContext());
+    checkf(bIsInitialized, TEXT("Cannot capture a frame if the capturer context is not initialized - is this being called before a frame has been rendered?"));
 
-    CUexternalMemory mappedExternalMemory = nullptr;
-
-    {
-        // generate a cudaExternalMemoryHandleDesc
-        CUDA_EXTERNAL_MEMORY_HANDLE_DESC cudaExtMemHandleDesc = {};
-        cudaExtMemHandleDesc.type = CU_EXTERNAL_MEMORY_HANDLE_TYPE_OPAQUE_FD;
-        cudaExtMemHandleDesc.handle.fd = fd;
-        cudaExtMemHandleDesc.size = VulkanTexture->Surface.GetAllocationOffset() + VulkanTexture->Surface.GetMemorySize();
-
-        // import external memory
-        auto result = FCUDAModule::CUDA().cuImportExternalMemory(&mappedExternalMemory, &cudaExtMemHandleDesc);
-        if(result != CUDA_SUCCESS)
-        {
-            UE_LOG(PixelStreamer, Error, TEXT("Failed to import external memory from vulkan error: %d"), result);
-        }
-    }
-
-    CUmipmappedArray mappedMipArray = nullptr;
-    CUarray mappedArray = nullptr;
-
-    {
-        CUDA_EXTERNAL_MEMORY_MIPMAPPED_ARRAY_DESC mipmapDesc = {};
-        mipmapDesc.numLevels = 1;
-        mipmapDesc.offset = VulkanTexture->Surface.GetAllocationOffset();
-        mipmapDesc.arrayDesc.Width = Texture->GetSizeX();
-        mipmapDesc.arrayDesc.Height = Texture->GetSizeY();
-        mipmapDesc.arrayDesc.Depth = 0;
-        mipmapDesc.arrayDesc.NumChannels = 4;
-        mipmapDesc.arrayDesc.Format = CU_AD_FORMAT_UNSIGNED_INT8;
-        mipmapDesc.arrayDesc.Flags = CUDA_ARRAY3D_SURFACE_LDST | CUDA_ARRAY3D_COLOR_ATTACHMENT;
-
-        // get the CUarray from the external memory
-        auto result = FCUDAModule::CUDA().cuExternalMemoryGetMappedMipmappedArray(&mappedMipArray, mappedExternalMemory, &mipmapDesc);
-        if(result != CUDA_SUCCESS)
-        {
-            UE_LOG(PixelStreamer, Error, TEXT("Failed to bind mipmappedArray error: %d"), result);
-        }
-    }
-
-    // get the CUarray from the external memory
-    CUresult mipMapArrGetLevelErr = FCUDAModule::CUDA().cuMipmappedArrayGetLevel(&mappedArray, mappedMipArray, 0);
-    if(mipMapArrGetLevelErr != CUDA_SUCCESS)
-    {
-        UE_LOG(PixelStreamer, Error, TEXT("Failed to bind to mip 0."));
-    }
-
-    FCUDAModule::CUDA().cuCtxPopCurrent(NULL);
-
-    InputFrame->SetTexture(mappedArray, [this, mappedArray, mappedMipArray, mappedExternalMemory, InputFrame](CUarray NativeTexture)
-    {
-        // free the cuda types
-        FCUDAModule::CUDA().cuCtxPushCurrent(FModuleManager::GetModuleChecked<FCUDAModule>("CUDA").GetCudaContext());
-
-        if(mappedArray)
-        {
-            auto result = FCUDAModule::CUDA().cuArrayDestroy(mappedArray);
-            if(result != CUDA_SUCCESS)
-            {
-                UE_LOG(PixelStreamer, Error, TEXT("Failed to destroy mappedArray: %d"), result);
-            }
-        }
-
-        if(mappedMipArray)
-        {
-            auto result = FCUDAModule::CUDA().cuMipmappedArrayDestroy(mappedMipArray);
-            if(result != CUDA_SUCCESS)
-            {
-                UE_LOG(PixelStreamer, Error, TEXT("Failed to destroy mappedMipArray: %d"), result);
-            }
-        }
-
-        if(mappedExternalMemory)
-        {
-            auto result = FCUDAModule::CUDA().cuDestroyExternalMemory(mappedExternalMemory);
-            if(result != CUDA_SUCCESS)
-            {
-                UE_LOG(PixelStreamer, Error, TEXT("Failed to destroy mappedExternalMemoryArray: %d"), result);
-            }
-        }
-
-        FCUDAModule::CUDA().cuCtxPopCurrent(NULL);
+	return [this]() -> const FTexture2DRHIRef
+	{
+		if (bIsTempDirty)
+		{
+			FScopeLock Lock(&CriticalSection);
 
-        // finally remove the input frame
-        BackBuffers.Remove(InputFrame);
-    });
+			EncoderTexture.Swap(TempTexture);
+			
+			bIsTempDirty = false;
+		}
 
-	BackBuffers.Add(InputFrame, Texture);
-    return Texture;
+		return EncoderTexture;
+	};
 }
diff --git a/Engine/Plugins/Media/PixelStreaming/Source/PixelStreaming/Private/VideoCapturerContext.h b/Engine/Plugins/Media/PixelStreaming/Source/PixelStreaming/Private/VideoCapturerContext.h
--- a/Engine/Plugins/Media/PixelStreaming/Source/PixelStreaming/Private/VideoCapturerContext.h
+++ b/Engine/Plugins/Media/PixelStreaming/Source/PixelStreaming/Private/VideoCapturerContext.h
@@ -9,28 +9,21 @@
 #include "RHI.h"
 #include "RenderResource.h"
 #include "ScreenRendering.h"
+#include "PixelStreamingFrameBuffer.h"
+#include "Containers/TripleBuffer.h"
+
+struct FVideoCaptureFrame
+{
+    FGPUFenceRHIRef CopyFence;
+    FTexture2DRHIRef Texture;
+	uint64 PreWaitingOnCopy;
+};
+
 
 // Context object for `FVideoCapturer`. The context contains backbuffers and the appropriates objects/methods to create frames for submission to 
 // the hardware encoders. The context is split out from the `FVideoCapturer` itself because context can be shared/passed around to capturers as they become active.
 class FVideoCapturerContext
 {
-    public:
-
-        class FCapturerInput
-        {
-            public:
-                FCapturerInput() 
-                    : InputFrame(nullptr)
-                    , Texture(){};
-
-                FCapturerInput(AVEncoder::FVideoEncoderInputFrame* InFrame, FTexture2DRHIRef InTexture)
-                    : InputFrame(InFrame)
-                    , Texture(InTexture){};
-
-            public:
-                AVEncoder::FVideoEncoderInputFrame* InputFrame;
-                TOptional<FTexture2DRHIRef> Texture;
-        };
 
     public:
         FVideoCapturerContext(int InCaptureWidth, int InCaptureHeight, bool bInFixedResolution);
@@ -38,25 +31,32 @@ class FVideoCapturerContext
         int GetCaptureHeight() const;
         bool IsFixedResolution() const;
         void SetCaptureResolution(int NewCaptureWidth, int NewCaptureHeight);
-        FVideoCapturerContext::FCapturerInput ObtainCapturerInput();
-        TSharedPtr<AVEncoder::FVideoEncoderInput> GetVideoEncoderInput() const;
+        int32 GetNextFrameId();
+        void CaptureFrame(const FTexture2DRHIRef& FrameBuffer);
+        FTextureObtainer RequestNewestCapturedFrame();
+        bool IsInitialized() const;
 
     private:
-        TSharedPtr<AVEncoder::FVideoEncoderInput> CreateVideoEncoderInput(int InWidth, int InHeight, bool bInFixedResolution);
+        FTexture2DRHIRef MakeTexture();
+        void Shutdown();
         void DeleteBackBuffers();
-        
-#if PLATFORM_WINDOWS        
-        FTexture2DRHIRef SetBackbufferTextureDX11(AVEncoder::FVideoEncoderInputFrame* InputFrame);
-        FTexture2DRHIRef SetBackbufferTextureDX12(AVEncoder::FVideoEncoderInputFrame* InputFrame);
-#endif // PLATFORM_WINDOWS
-
-        FTexture2DRHIRef SetBackbufferTexturePureVulkan(AVEncoder::FVideoEncoderInputFrame* InputFrame);
-        FTexture2DRHIRef SetBackbufferTextureCUDAVulkan(AVEncoder::FVideoEncoderInputFrame* InputFrame);
+		void SwapBackBuffers();
 
     private:
         int CaptureWidth;
         int CaptureHeight;
         bool bFixedResolution;
-        TSharedPtr<AVEncoder::FVideoEncoderInput> VideoEncoderInput = nullptr;
-        TMap<AVEncoder::FVideoEncoderInputFrame*, FTexture2DRHIRef> BackBuffers;
+
+        FVideoCaptureFrame EvenFrame;
+        FVideoCaptureFrame OddFrame;
+        FTexture2DRHIRef TempTexture;
+        FTexture2DRHIRef EncoderTexture;
+
+		FThreadSafeBool bIsTempDirty;
+
+        FThreadSafeCounter NextFrameID;
+        FThreadSafeBool bIsInitialized;
+        FCriticalSection CriticalSection;
+
+        bool bIsEvenFrame;
 };
diff --git a/Engine/Plugins/Media/PixelStreaming/Source/PixelStreaming/Public/PixelStreamerInputComponent.h b/Engine/Plugins/Media/PixelStreaming/Source/PixelStreaming/Public/PixelStreamerInputComponent.h
--- a/Engine/Plugins/Media/PixelStreaming/Source/PixelStreaming/Public/PixelStreamerInputComponent.h
+++ b/Engine/Plugins/Media/PixelStreaming/Source/PixelStreaming/Public/PixelStreamerInputComponent.h
@@ -42,7 +42,7 @@ public:
 	 * @param Descriptor - The command JSON descriptor.
 	 * @return True if the command was successfully executed.
 	 */
-	bool OnCommand(const FString& Descriptor);
+	static bool OnCommand(const FString& Descriptor);
 
 	/**
 	 * Send a response back to the source of the UI interactions.
@@ -64,6 +64,7 @@ public:
 	 */
 	UFUNCTION(BlueprintPure, Category = "PixelStreamer Input")
 	void GetJsonStringValue(FString Descriptor, FString FieldName, FString& StringValue, bool& Success);
+	static void ExtractJsonFromDescriptor(FString Descriptor, FString FieldName, FString& StringValue, bool& Success);
 
 	/**
 	 * Helper function to add a string field to a JSON descriptor. This produces
@@ -76,6 +77,7 @@ public:
 	 */
 	UFUNCTION(BlueprintPure, Category = "PixelStreamer Input")
 	void AddJsonStringValue(const FString& Descriptor, FString FieldName, FString StringValue, FString& NewDescriptor, bool& Success);
+	static void ExtendJsonWithField(const FString& Descriptor, FString FieldName, FString StringValue, FString& NewDescriptor, bool& Success);
 
 private:
 	// For convenience we keep a reference to the Pixel Streaming plugin.
diff --git a/Engine/Source/Runtime/AVEncoder/Private/VideoEncoderInput.cpp b/Engine/Source/Runtime/AVEncoder/Private/VideoEncoderInput.cpp
--- a/Engine/Source/Runtime/AVEncoder/Private/VideoEncoderInput.cpp
+++ b/Engine/Source/Runtime/AVEncoder/Private/VideoEncoderInput.cpp
@@ -31,1025 +31,1043 @@ FString GetGUID()
 
 namespace AVEncoder
 {
+	// *** FVideoEncoderInput *************************************************************************
 
-// *** FVideoEncoderInput *************************************************************************
+	// --- construct video encoder input based on expected input frame format -------------------------
 
-// --- construct video encoder input based on expected input frame format -------------------------
-
-TSharedPtr<FVideoEncoderInput> FVideoEncoderInput::CreateDummy(uint32 InWidth, uint32 InHeight, bool bIsResizable)
-{
-	TSharedPtr<FVideoEncoderInputImpl>	Input = MakeShared<FVideoEncoderInputImpl>();
-	Input->bIsResizable = bIsResizable;
-
-	if (!Input->SetupForDummy(InWidth, InHeight))
+	TSharedPtr<FVideoEncoderInput> FVideoEncoderInput::CreateDummy(uint32 InWidth, uint32 InHeight, bool bIsResizable)
 	{
-		Input.Reset();
-	}
-	return Input;
-}
+		TSharedPtr<FVideoEncoderInputImpl>	Input = MakeShared<FVideoEncoderInputImpl>();
+		Input->bIsResizable = bIsResizable;
 
-TSharedPtr<FVideoEncoderInput> FVideoEncoderInput::CreateForYUV420P(uint32 InWidth, uint32 InHeight, bool bIsResizable)
-{
-	TSharedPtr<FVideoEncoderInputImpl>	Input = MakeShared<FVideoEncoderInputImpl>();
-	Input->bIsResizable = bIsResizable;
+		if (!Input->SetupForDummy(InWidth, InHeight))
+		{
+			Input.Reset();
+		}
+		return Input;
+	}
 
-	if (!Input->SetupForYUV420P(InWidth, InHeight))
+	TSharedPtr<FVideoEncoderInput> FVideoEncoderInput::CreateForYUV420P(uint32 InWidth, uint32 InHeight, bool bIsResizable)
 	{
-		Input.Reset();
+		TSharedPtr<FVideoEncoderInputImpl>	Input = MakeShared<FVideoEncoderInputImpl>();
+		Input->bIsResizable = bIsResizable;
+
+		if (!Input->SetupForYUV420P(InWidth, InHeight))
+		{
+			Input.Reset();
+		}
+		return Input;
 	}
-	return Input;
-}
 
-TSharedPtr<FVideoEncoderInput> FVideoEncoderInput::CreateForD3D11(void* InApplicationD3DDevice, uint32 InWidth, uint32 InHeight, bool bIsResizable, bool IsShared)
-{
-	TSharedPtr<FVideoEncoderInputImpl>	Input = MakeShared<FVideoEncoderInputImpl>();
+	TSharedPtr<FVideoEncoderInput> FVideoEncoderInput::CreateForD3D11(void* InApplicationD3DDevice, uint32 InWidth, uint32 InHeight, bool bIsResizable, bool IsShared)
+	{
+		TSharedPtr<FVideoEncoderInputImpl>	Input = MakeShared<FVideoEncoderInputImpl>();
 
 #if PLATFORM_WINDOWS
-	Input->bIsResizable = bIsResizable;
+		Input->bIsResizable = bIsResizable;
 
-	if (IsShared)
-	{
-		if (!Input->SetupForD3D11Shared(static_cast<ID3D11Device*>(InApplicationD3DDevice), InWidth, InHeight))
+		if (IsShared)
 		{
-			Input.Reset();
+			if (!Input->SetupForD3D11Shared(static_cast<ID3D11Device*>(InApplicationD3DDevice), InWidth, InHeight))
+			{
+				Input.Reset();
+			}
 		}
-	}
-	else
-	{
-		if (!Input->SetupForD3D11(static_cast<ID3D11Device*>(InApplicationD3DDevice), InWidth, InHeight))
+		else
 		{
-			Input.Reset();
+			if (!Input->SetupForD3D11(static_cast<ID3D11Device*>(InApplicationD3DDevice), InWidth, InHeight))
+			{
+				Input.Reset();
+			}
 		}
-	}
 
 #endif
 
-	return Input;
-}
+		return Input;
+	}
 
 
-TSharedPtr<FVideoEncoderInput> FVideoEncoderInput::CreateForD3D12(void* InApplicationD3DDevice, uint32 InWidth, uint32 InHeight, bool bIsResizable, bool IsShared)
-{
-	TSharedPtr<FVideoEncoderInputImpl>	Input = MakeShared<FVideoEncoderInputImpl>();
+	TSharedPtr<FVideoEncoderInput> FVideoEncoderInput::CreateForD3D12(void* InApplicationD3DDevice, uint32 InWidth, uint32 InHeight, bool bIsResizable, bool IsShared)
+	{
+		TSharedPtr<FVideoEncoderInputImpl>	Input = MakeShared<FVideoEncoderInputImpl>();
 
 #if PLATFORM_WINDOWS
-	Input->bIsResizable = bIsResizable;
+		Input->bIsResizable = bIsResizable;
 
-	if(IsShared)
-	{
-		if (!Input->SetupForD3D12Shared(static_cast<ID3D12Device*>(InApplicationD3DDevice), InWidth, InHeight))
+		if (IsShared)
 		{
-			Input.Reset();
+			if (!Input->SetupForD3D12Shared(static_cast<ID3D12Device*>(InApplicationD3DDevice), InWidth, InHeight))
+			{
+				Input.Reset();
+			}
 		}
-	}
-	else
-	{
-		if (!Input->SetupForD3D12(static_cast<ID3D12Device*>(InApplicationD3DDevice), InWidth, InHeight))
+		else
 		{
-			Input.Reset();
+			if (!Input->SetupForD3D12(static_cast<ID3D12Device*>(InApplicationD3DDevice), InWidth, InHeight))
+			{
+				Input.Reset();
+			}
 		}
-	}
-	
+
 #endif
 
-	return Input;
-}
+		return Input;
+	}
 
-TSharedPtr<FVideoEncoderInput> FVideoEncoderInput::CreateForCUDA(void* InApplicationContext, uint32 InWidth, uint32 InHeight, bool bIsResizable)
-{
-	TSharedPtr<FVideoEncoderInputImpl> Input = MakeShared<FVideoEncoderInputImpl>();
-	Input->bIsResizable = bIsResizable;
-		
-	if (!Input->SetupForCUDA(reinterpret_cast<CUcontext>(InApplicationContext), InWidth, InHeight))
+	TSharedPtr<FVideoEncoderInput> FVideoEncoderInput::CreateForCUDA(void* InApplicationContext, uint32 InWidth, uint32 InHeight, bool bIsResizable)
 	{
-		Input.Reset();
+		TSharedPtr<FVideoEncoderInputImpl> Input = MakeShared<FVideoEncoderInputImpl>();
+		Input->bIsResizable = bIsResizable;
+
+		if (!Input->SetupForCUDA(reinterpret_cast<CUcontext>(InApplicationContext), InWidth, InHeight))
+		{
+			Input.Reset();
+		}
+		return Input;
 	}
-	return Input;
-}
 
 #if PLATFORM_DESKTOP && !PLATFORM_APPLE
-TSharedPtr<FVideoEncoderInput> FVideoEncoderInput::CreateForVulkan(void* InApplicationVulkanData, uint32 InWidth, uint32 InHeight, bool bIsResizable)
-{
-	TSharedPtr<FVideoEncoderInputImpl>	Input = MakeShared<FVideoEncoderInputImpl>();
-	Input->bIsResizable = bIsResizable;
+	TSharedPtr<FVideoEncoderInput> FVideoEncoderInput::CreateForVulkan(void* InApplicationVulkanData, uint32 InWidth, uint32 InHeight, bool bIsResizable)
+	{
+		TSharedPtr<FVideoEncoderInputImpl>	Input = MakeShared<FVideoEncoderInputImpl>();
+		Input->bIsResizable = bIsResizable;
 
-	FVulkanDataStruct* VulkanData = static_cast<FVulkanDataStruct*>(InApplicationVulkanData);
+		FVulkanDataStruct* VulkanData = static_cast<FVulkanDataStruct*>(InApplicationVulkanData);
 
-	if (!Input->SetupForVulkan(	VulkanData->VulkanInstance, VulkanData->VulkanPhysicalDevice, VulkanData->VulkanDevice, InWidth, InHeight))
-	{
-		Input.Reset();
-	}
+		if (!Input->SetupForVulkan(VulkanData->VulkanInstance, VulkanData->VulkanPhysicalDevice, VulkanData->VulkanDevice, InWidth, InHeight))
+		{
+			Input.Reset();
+		}
 
-	return Input;
-}
+		return Input;
+	}
 #endif
 
-void FVideoEncoderInput::SetResolution(uint32 InWidth, uint32 InHeight)
-{
-	Width = InWidth;
-	Height = InHeight;
-}
+	void FVideoEncoderInput::SetResolution(uint32 InWidth, uint32 InHeight)
+	{
+		Width = InWidth;
+		Height = InHeight;
+	}
 
-void FVideoEncoderInput::SetMaxNumBuffers(uint32 InMaxNumBuffers)
-{
-	MaxNumBuffers = InMaxNumBuffers;
-}
+	void FVideoEncoderInput::SetMaxNumBuffers(uint32 InMaxNumBuffers)
+	{
+		MaxNumBuffers = InMaxNumBuffers;
+	}
 
-// --- encoder input frames -----------------------------------------------------------------------
+	// --- encoder input frames -----------------------------------------------------------------------
 
 
 
-// *** FVideoEncoderInputImpl *********************************************************************
+	// *** FVideoEncoderInputImpl *********************************************************************
 
-FVideoEncoderInputImpl::~FVideoEncoderInputImpl()
-{
+	FVideoEncoderInputImpl::~FVideoEncoderInputImpl()
 	{
-		FScopeLock						Guard(&ProtectFrames);
-		if (ActiveFrames.Num() > 0)
 		{
-			UE_LOG(LogVideoEncoder, Error, TEXT("There are still %d active input frames."), ActiveFrames.Num());
-		}
-		
-		check(ActiveFrames.Num() == 0);
-		
-		while (!AvailableFrames.IsEmpty())
-		{
-			FVideoEncoderInputFrameImpl* Frame = nullptr;
-			AvailableFrames.Dequeue(Frame);
-			delete Frame;
-		}
-	}
-#if PLATFORM_WINDOWS
-//	DEBUG_D3D11_REPORT_LIVE_DEVICE_OBJECT(FrameInfoD3D.EncoderDeviceD3D11);
-#endif
-}
+			FScopeLock Guard(&ProtectFrames);
+			if (ActiveFrames.Num() > 0)
+			{
+				UE_LOG(LogVideoEncoder, Error, TEXT("There are still %d active input frames."), ActiveFrames.Num());
+			}
 
-bool FVideoEncoderInputImpl::SetupForDummy(uint32 InWidth, uint32 InHeight)
-{
-	FrameFormat = EVideoFrameFormat::Undefined;
-	this->Width = InWidth;
-	this->Height = InHeight;
-	return true;
-}
+			check(ActiveFrames.Num() == 0);
 
-bool FVideoEncoderInputImpl::SetupForYUV420P(uint32 InWidth, uint32 InHeight)
-{
-	FrameFormat = EVideoFrameFormat::YUV420P;
-	this->Width = InWidth;
-	this->Height = InHeight;
-	FrameInfoYUV420P.StrideY = InWidth;
-	FrameInfoYUV420P.StrideU = (InWidth + 1) / 2;
-	FrameInfoYUV420P.StrideV = (InWidth + 1) / 2;
-
-	CollectAvailableEncoders();
-	return true;
-}
-
-bool FVideoEncoderInputImpl::SetupForD3D11(void* InApplicationD3DDevice, uint32 InWidth, uint32 InHeight)
-{
+			while (!AvailableFrames.IsEmpty())
+			{
+				FVideoEncoderInputFrameImpl* Frame = nullptr;
+				AvailableFrames.Dequeue(Frame);
+				delete Frame;
+			}
+		}
 #if PLATFORM_WINDOWS
-	TRefCountPtr<IDXGIDevice>	DXGIDevice;
-	TRefCountPtr<IDXGIAdapter>	Adapter;
-
-	HRESULT		Result = static_cast<ID3D11Device*>(InApplicationD3DDevice)->QueryInterface(__uuidof(IDXGIDevice), (void**)DXGIDevice.GetInitReference());
-	if (Result != S_OK)
-	{
-		UE_LOG(LogVideoEncoder, Error, TEXT("ID3D11Device::QueryInterface() failed 0x%X - %s."), Result, *GetComErrorDescription(Result));
-		return false;
+		//	DEBUG_D3D11_REPORT_LIVE_DEVICE_OBJECT(FrameInfoD3D.EncoderDeviceD3D11);
+#endif
 	}
-	else if ((Result = DXGIDevice->GetAdapter(Adapter.GetInitReference())) != S_OK)
+
+	bool FVideoEncoderInputImpl::SetupForDummy(uint32 InWidth, uint32 InHeight)
 	{
-		UE_LOG(LogVideoEncoder, Error, TEXT("DXGIDevice::GetAdapter() failed 0x%X - %s."), Result, *GetComErrorDescription(Result));
-		return false;
+		FrameFormat = EVideoFrameFormat::Undefined;
+		this->Width = InWidth;
+		this->Height = InHeight;
+		return true;
 	}
 
-	uint32				DeviceFlags = 0;
-	D3D_FEATURE_LEVEL	FeatureLevel = D3D_FEATURE_LEVEL_11_0;
-	D3D_FEATURE_LEVEL	ActualFeatureLevel;
-
-	if ((Result = D3D11CreateDevice(
-		Adapter,
-		D3D_DRIVER_TYPE_UNKNOWN,
-		NULL,
-		DeviceFlags,
-		&FeatureLevel,
-		1,
-		D3D11_SDK_VERSION,
-		FrameInfoD3D.EncoderDeviceD3D11.GetInitReference(),
-		&ActualFeatureLevel,
-		FrameInfoD3D.EncoderDeviceContextD3D11.GetInitReference())) != S_OK)
+	bool FVideoEncoderInputImpl::SetupForYUV420P(uint32 InWidth, uint32 InHeight)
 	{
-		UE_LOG(LogVideoEncoder, Error, TEXT("D3D11CreateDevice() failed 0x%X - %s."), Result, *GetComErrorDescription(Result));
-		return false;
+		FrameFormat = EVideoFrameFormat::YUV420P;
+		this->Width = InWidth;
+		this->Height = InHeight;
+		FrameInfoYUV420P.StrideY = InWidth;
+		FrameInfoYUV420P.StrideU = (InWidth + 1) / 2;
+		FrameInfoYUV420P.StrideV = (InWidth + 1) / 2;
+
+		CollectAvailableEncoders();
+		return true;
 	}
-	DEBUG_SET_D3D11_OBJECT_NAME(FrameInfoD3D.EncoderDeviceD3D11, "FVideoEncoderInputImpl");
-	DEBUG_SET_D3D11_OBJECT_NAME(FrameInfoD3D.EncoderDeviceContextD3D11, "FVideoEncoderInputImpl");
 
-	FrameFormat = EVideoFrameFormat::D3D11_R8G8B8A8_UNORM;
-	this->Width = InWidth;
-	this->Height = InHeight;
+	bool FVideoEncoderInputImpl::SetupForD3D11(void* InApplicationD3DDevice, uint32 InWidth, uint32 InHeight)
+	{
+#if PLATFORM_WINDOWS
+		TRefCountPtr<IDXGIDevice>	DXGIDevice;
+		TRefCountPtr<IDXGIAdapter>	Adapter;
 
-	CollectAvailableEncoders();
-	return true;
+		HRESULT	Result = static_cast<ID3D11Device*>(InApplicationD3DDevice)->QueryInterface(__uuidof(IDXGIDevice), (void**)DXGIDevice.GetInitReference());
+		if (Result != S_OK)
+		{
+			UE_LOG(LogVideoEncoder, Error, TEXT("ID3D11Device::QueryInterface() failed 0x%X - %s."), Result, *GetComErrorDescription(Result));
+			return false;
+		}
+		else if ((Result = DXGIDevice->GetAdapter(Adapter.GetInitReference())) != S_OK)
+		{
+			UE_LOG(LogVideoEncoder, Error, TEXT("DXGIDevice::GetAdapter() failed 0x%X - %s."), Result, *GetComErrorDescription(Result));
+			return false;
+		}
 
-#endif
-	return false;
-}
+		uint32				DeviceFlags = 0;
+		D3D_FEATURE_LEVEL	FeatureLevel = D3D_FEATURE_LEVEL_11_0;
+		D3D_FEATURE_LEVEL	ActualFeatureLevel;
+
+		if ((Result = D3D11CreateDevice(
+			Adapter,
+			D3D_DRIVER_TYPE_UNKNOWN,
+			NULL,
+			DeviceFlags,
+			&FeatureLevel,
+			1,
+			D3D11_SDK_VERSION,
+			FrameInfoD3D.EncoderDeviceD3D11.GetInitReference(),
+			&ActualFeatureLevel,
+			FrameInfoD3D.EncoderDeviceContextD3D11.GetInitReference())) != S_OK)
+		{
+			UE_LOG(LogVideoEncoder, Error, TEXT("D3D11CreateDevice() failed 0x%X - %s."), Result, *GetComErrorDescription(Result));
+			return false;
+		}
+		DEBUG_SET_D3D11_OBJECT_NAME(FrameInfoD3D.EncoderDeviceD3D11, "FVideoEncoderInputImpl");
+		DEBUG_SET_D3D11_OBJECT_NAME(FrameInfoD3D.EncoderDeviceContextD3D11, "FVideoEncoderInputImpl");
 
-bool FVideoEncoderInputImpl::SetupForD3D11Shared(void* InApplicationD3DDevice, uint32 InWidth, uint32 InHeight)
-{
-#if PLATFORM_WINDOWS
+		FrameFormat = EVideoFrameFormat::D3D11_R8G8B8A8_UNORM;
+		this->Width = InWidth;
+		this->Height = InHeight;
 
-	TRefCountPtr<IDXGIDevice>	DXGIDevice;
-	TRefCountPtr<IDXGIAdapter>	Adapter;
+		CollectAvailableEncoders();
+		return true;
 
-	HRESULT		Result = static_cast<ID3D11Device*>(InApplicationD3DDevice)->QueryInterface(__uuidof(IDXGIDevice), (void**)DXGIDevice.GetInitReference());
-	if (Result != S_OK)
-	{
-		UE_LOG(LogVideoEncoder, Error, TEXT("ID3D11Device::QueryInterface() failed 0x%X - %s."), Result, *GetComErrorDescription(Result));
-		return false;
-	}
-	else if ((Result = DXGIDevice->GetAdapter(Adapter.GetInitReference())) != S_OK)
-	{
-		UE_LOG(LogVideoEncoder, Error, TEXT("DXGIDevice::GetAdapter() failed 0x%X - %s."), Result, *GetComErrorDescription(Result));
+#endif
 		return false;
 	}
 
-	uint32				DeviceFlags = 0;
-	D3D_FEATURE_LEVEL	FeatureLevel = D3D_FEATURE_LEVEL_11_1;
-	D3D_FEATURE_LEVEL	ActualFeatureLevel;
-
-	if ((Result = D3D11CreateDevice(
-		Adapter,
-		D3D_DRIVER_TYPE_UNKNOWN,
-		NULL,
-		DeviceFlags,
-		&FeatureLevel,
-		1,
-		D3D11_SDK_VERSION,
-		FrameInfoD3D.EncoderDeviceD3D11.GetInitReference(),
-		&ActualFeatureLevel,
-		FrameInfoD3D.EncoderDeviceContextD3D11.GetInitReference())) != S_OK)
+	bool FVideoEncoderInputImpl::SetupForD3D11Shared(void* InApplicationD3DDevice, uint32 InWidth, uint32 InHeight)
 	{
-		UE_LOG(LogVideoEncoder, Error, TEXT("D3D11CreateDevice() failed 0x%X - %s."), Result, *GetComErrorDescription(Result));
-		return false;
-	}
-	DEBUG_SET_D3D11_OBJECT_NAME(FrameInfoD3D.EncoderDeviceD3D11, "FVideoEncoderInputImpl");
-	DEBUG_SET_D3D11_OBJECT_NAME(FrameInfoD3D.EncoderDeviceContextD3D11, "FVideoEncoderInputImpl");
+#if PLATFORM_WINDOWS
 
-	FrameFormat = EVideoFrameFormat::D3D11_R8G8B8A8_UNORM;
-	this->Width = InWidth;
-	this->Height = InHeight;
+		TRefCountPtr<IDXGIDevice>	DXGIDevice;
+		TRefCountPtr<IDXGIAdapter>	Adapter;
 
-	CollectAvailableEncoders();
-	return true;
+		HRESULT		Result = static_cast<ID3D11Device*>(InApplicationD3DDevice)->QueryInterface(__uuidof(IDXGIDevice), (void**)DXGIDevice.GetInitReference());
+		if (Result != S_OK)
+		{
+			UE_LOG(LogVideoEncoder, Error, TEXT("ID3D11Device::QueryInterface() failed 0x%X - %s."), Result, *GetComErrorDescription(Result));
+			return false;
+		}
+		else if ((Result = DXGIDevice->GetAdapter(Adapter.GetInitReference())) != S_OK)
+		{
+			UE_LOG(LogVideoEncoder, Error, TEXT("DXGIDevice::GetAdapter() failed 0x%X - %s."), Result, *GetComErrorDescription(Result));
+			return false;
+		}
 
-#endif
-	return false;
-}
+		uint32				DeviceFlags = 0;
+		D3D_FEATURE_LEVEL	FeatureLevel = D3D_FEATURE_LEVEL_11_1;
+		D3D_FEATURE_LEVEL	ActualFeatureLevel;
+
+		if ((Result = D3D11CreateDevice(
+			Adapter,
+			D3D_DRIVER_TYPE_UNKNOWN,
+			NULL,
+			DeviceFlags,
+			&FeatureLevel,
+			1,
+			D3D11_SDK_VERSION,
+			FrameInfoD3D.EncoderDeviceD3D11.GetInitReference(),
+			&ActualFeatureLevel,
+			FrameInfoD3D.EncoderDeviceContextD3D11.GetInitReference())) != S_OK)
+		{
+			UE_LOG(LogVideoEncoder, Error, TEXT("D3D11CreateDevice() failed 0x%X - %s."), Result, *GetComErrorDescription(Result));
+			return false;
+		}
+		DEBUG_SET_D3D11_OBJECT_NAME(FrameInfoD3D.EncoderDeviceD3D11, "FVideoEncoderInputImpl");
+		DEBUG_SET_D3D11_OBJECT_NAME(FrameInfoD3D.EncoderDeviceContextD3D11, "FVideoEncoderInputImpl");
 
-bool FVideoEncoderInputImpl::SetupForD3D12(void* InApplicationD3DDevice, uint32 InWidth, uint32 InHeight)
-{
-#if PLATFORM_WINDOWS
+		FrameFormat = EVideoFrameFormat::D3D11_R8G8B8A8_UNORM;
+		this->Width = InWidth;
+		this->Height = InHeight;
 
-	LUID						AdapterLuid = static_cast<ID3D12Device*>(InApplicationD3DDevice)->GetAdapterLuid();
-	TRefCountPtr<IDXGIFactory4>	DXGIFactory;
-	HRESULT						Result;
-	if ((Result = CreateDXGIFactory(IID_PPV_ARGS(DXGIFactory.GetInitReference()))) != S_OK)
-	{
-		UE_LOG(LogVideoEncoder, Error, TEXT("CreateDXGIFactory() failed 0x%X - %s."), Result, *GetComErrorDescription(Result));
-		return false;
-	}
-	// get the adapter game uses to render
-	TRefCountPtr<IDXGIAdapter>	Adapter;
-	if ((Result = DXGIFactory->EnumAdapterByLuid(AdapterLuid, IID_PPV_ARGS(Adapter.GetInitReference()))) != S_OK)
-	{
-		UE_LOG(LogVideoEncoder, Error, TEXT("DXGIFactory::EnumAdapterByLuid() failed 0x%X - %s."), Result, *GetComErrorDescription(Result));
+		CollectAvailableEncoders();
+		return true;
+
+#endif
 		return false;
 	}
 
-	uint32				DeviceFlags = 0;
-	D3D_FEATURE_LEVEL	FeatureLevel = D3D_FEATURE_LEVEL_12_0; // TODO get this from the adaptor support
-	if ((Result = D3D12CreateDevice(Adapter, FeatureLevel, IID_PPV_ARGS(FrameInfoD3D.EncoderDeviceD3D12.GetInitReference()))) != S_OK)
+	bool FVideoEncoderInputImpl::SetupForD3D12(void* InApplicationD3DDevice, uint32 InWidth, uint32 InHeight)
 	{
-		UE_LOG(LogVideoEncoder, Error, TEXT("D3D11CreateDevice() failed 0x%X - %s."), Result, *GetComErrorDescription(Result));
-		return false;
-	}
+#if PLATFORM_WINDOWS
 
-	FrameFormat = EVideoFrameFormat::D3D12_R8G8B8A8_UNORM;
-	this->Width = InWidth;
-	this->Height = InHeight;
+		LUID						AdapterLuid = static_cast<ID3D12Device*>(InApplicationD3DDevice)->GetAdapterLuid();
+		TRefCountPtr<IDXGIFactory4>	DXGIFactory;
+		HRESULT						Result;
+		if ((Result = CreateDXGIFactory(IID_PPV_ARGS(DXGIFactory.GetInitReference()))) != S_OK)
+		{
+			UE_LOG(LogVideoEncoder, Error, TEXT("CreateDXGIFactory() failed 0x%X - %s."), Result, *GetComErrorDescription(Result));
+			return false;
+		}
+		// get the adapter game uses to render
+		TRefCountPtr<IDXGIAdapter>	Adapter;
+		if ((Result = DXGIFactory->EnumAdapterByLuid(AdapterLuid, IID_PPV_ARGS(Adapter.GetInitReference()))) != S_OK)
+		{
+			UE_LOG(LogVideoEncoder, Error, TEXT("DXGIFactory::EnumAdapterByLuid() failed 0x%X - %s."), Result, *GetComErrorDescription(Result));
+			return false;
+		}
 
-	CollectAvailableEncoders();
-	return true;
+		uint32				DeviceFlags = 0;
+		D3D_FEATURE_LEVEL	FeatureLevel = D3D_FEATURE_LEVEL_12_0; // TODO get this from the adaptor support
+		if ((Result = D3D12CreateDevice(Adapter, FeatureLevel, IID_PPV_ARGS(FrameInfoD3D.EncoderDeviceD3D12.GetInitReference()))) != S_OK)
+		{
+			UE_LOG(LogVideoEncoder, Error, TEXT("D3D11CreateDevice() failed 0x%X - %s."), Result, *GetComErrorDescription(Result));
+			return false;
+		}
 
-#endif
+		FrameFormat = EVideoFrameFormat::D3D12_R8G8B8A8_UNORM;
+		this->Width = InWidth;
+		this->Height = InHeight;
 
-	return false;
-}
+		CollectAvailableEncoders();
+		return true;
 
-bool FVideoEncoderInputImpl::SetupForD3D12Shared(void* InApplicationD3DDevice, uint32 InWidth, uint32 InHeight)
-{
-#if PLATFORM_WINDOWS
+#endif
 
-	LUID						AdapterLuid = static_cast<ID3D12Device*>(InApplicationD3DDevice)->GetAdapterLuid();
-	TRefCountPtr<IDXGIFactory4>	DXGIFactory;
-	HRESULT						Result;
-	if ((Result = CreateDXGIFactory(IID_PPV_ARGS(DXGIFactory.GetInitReference()))) != S_OK)
-	{
-		UE_LOG(LogVideoEncoder, Error, TEXT("CreateDXGIFactory() failed 0x%X - %s."), Result, *GetComErrorDescription(Result));
-		return false;
-	}
-	// get the adapter game uses to render
-	TRefCountPtr<IDXGIAdapter>	Adapter;
-	if ((Result = DXGIFactory->EnumAdapterByLuid(AdapterLuid, IID_PPV_ARGS(Adapter.GetInitReference()))) != S_OK)
-	{
-		UE_LOG(LogVideoEncoder, Error, TEXT("DXGIFactory::EnumAdapterByLuid() failed 0x%X - %s."), Result, *GetComErrorDescription(Result));
 		return false;
 	}
 
-	uint32				DeviceFlags = 0;
-	D3D_FEATURE_LEVEL	FeatureLevel = D3D_FEATURE_LEVEL_11_1;
-	D3D_FEATURE_LEVEL	ActualFeatureLevel;
-	if ((Result = D3D11CreateDevice(
-		Adapter,
-		D3D_DRIVER_TYPE_UNKNOWN,
-		NULL,
-		DeviceFlags,
-		&FeatureLevel,
-		1,
-		D3D11_SDK_VERSION,
-		FrameInfoD3D.EncoderDeviceD3D11.GetInitReference(),
-		&ActualFeatureLevel,
-		FrameInfoD3D.EncoderDeviceContextD3D11.GetInitReference())) != S_OK)
+	bool FVideoEncoderInputImpl::SetupForD3D12Shared(void* InApplicationD3DDevice, uint32 InWidth, uint32 InHeight)
 	{
-		UE_LOG(LogVideoEncoder, Error, TEXT("D3D11CreateDevice() failed 0x%X - %s."), Result, *GetComErrorDescription(Result));
-		return false;
-	}
+#if PLATFORM_WINDOWS
 
-	if (ActualFeatureLevel != D3D_FEATURE_LEVEL_11_1)
-	{
-		UE_LOG(LogVideoEncoder, Error, TEXT("D3D11CreateDevice() - failed to create device w/ feature level 11.1 - needed to encode textures from D3D12."));
-		FrameInfoD3D.EncoderDeviceD3D11.SafeRelease();
-		FrameInfoD3D.EncoderDeviceContextD3D11.SafeRelease();
-		return false;
-	}
+		LUID						AdapterLuid = static_cast<ID3D12Device*>(InApplicationD3DDevice)->GetAdapterLuid();
+		TRefCountPtr<IDXGIFactory4>	DXGIFactory;
+		HRESULT						Result;
+		if ((Result = CreateDXGIFactory(IID_PPV_ARGS(DXGIFactory.GetInitReference()))) != S_OK)
+		{
+			UE_LOG(LogVideoEncoder, Error, TEXT("CreateDXGIFactory() failed 0x%X - %s."), Result, *GetComErrorDescription(Result));
+			return false;
+		}
+		// get the adapter game uses to render
+		TRefCountPtr<IDXGIAdapter>	Adapter;
+		if ((Result = DXGIFactory->EnumAdapterByLuid(AdapterLuid, IID_PPV_ARGS(Adapter.GetInitReference()))) != S_OK)
+		{
+			UE_LOG(LogVideoEncoder, Error, TEXT("DXGIFactory::EnumAdapterByLuid() failed 0x%X - %s."), Result, *GetComErrorDescription(Result));
+			return false;
+		}
+
+		uint32				DeviceFlags = 0;
+		D3D_FEATURE_LEVEL	FeatureLevel = D3D_FEATURE_LEVEL_11_1;
+		D3D_FEATURE_LEVEL	ActualFeatureLevel;
+		if ((Result = D3D11CreateDevice(
+			Adapter,
+			D3D_DRIVER_TYPE_UNKNOWN,
+			NULL,
+			DeviceFlags,
+			&FeatureLevel,
+			1,
+			D3D11_SDK_VERSION,
+			FrameInfoD3D.EncoderDeviceD3D11.GetInitReference(),
+			&ActualFeatureLevel,
+			FrameInfoD3D.EncoderDeviceContextD3D11.GetInitReference())) != S_OK)
+		{
+			UE_LOG(LogVideoEncoder, Error, TEXT("D3D11CreateDevice() failed 0x%X - %s."), Result, *GetComErrorDescription(Result));
+			return false;
+		}
+
+		if (ActualFeatureLevel != D3D_FEATURE_LEVEL_11_1)
+		{
+			UE_LOG(LogVideoEncoder, Error, TEXT("D3D11CreateDevice() - failed to create device w/ feature level 11.1 - needed to encode textures from D3D12."));
+			FrameInfoD3D.EncoderDeviceD3D11.SafeRelease();
+			FrameInfoD3D.EncoderDeviceContextD3D11.SafeRelease();
+			return false;
+		}
 
-	DEBUG_SET_D3D11_OBJECT_NAME(FrameInfoD3D.EncoderDeviceD3D11, "FVideoEncoderInputImpl");
-	DEBUG_SET_D3D11_OBJECT_NAME(FrameInfoD3D.EncoderDeviceContextD3D11, "FVideoEncoderInputImpl");
+		DEBUG_SET_D3D11_OBJECT_NAME(FrameInfoD3D.EncoderDeviceD3D11, "FVideoEncoderInputImpl");
+		DEBUG_SET_D3D11_OBJECT_NAME(FrameInfoD3D.EncoderDeviceContextD3D11, "FVideoEncoderInputImpl");
 
-	FrameFormat = EVideoFrameFormat::D3D11_R8G8B8A8_UNORM;
-	this->Width = InWidth;
-	this->Height = InHeight;
+		FrameFormat = EVideoFrameFormat::D3D11_R8G8B8A8_UNORM;
+		this->Width = InWidth;
+		this->Height = InHeight;
 
-	CollectAvailableEncoders();
-	return true;
+		CollectAvailableEncoders();
+		return true;
 
 #endif
 
-	return false;
-}
+		return false;
+	}
 
 
-bool FVideoEncoderInputImpl::SetupForCUDA(void* InApplicationContext, uint32 InWidth, uint32 InHeight)
-{
-	FrameInfoCUDA.EncoderContextCUDA = static_cast<CUcontext>(InApplicationContext);
+	bool FVideoEncoderInputImpl::SetupForCUDA(void* InApplicationContext, uint32 InWidth, uint32 InHeight)
+	{
+		FrameInfoCUDA.EncoderContextCUDA = static_cast<CUcontext>(InApplicationContext);
 
-	FrameFormat = EVideoFrameFormat::CUDA_R8G8B8A8_UNORM;
-	this->Width = InWidth;
-	this->Height = InHeight;
+		FrameFormat = EVideoFrameFormat::CUDA_R8G8B8A8_UNORM;
+		this->Width = InWidth;
+		this->Height = InHeight;
 
-	CollectAvailableEncoders();
-	return true;
-}
+		CollectAvailableEncoders();
+		return true;
+	}
 
 #if PLATFORM_DESKTOP && !PLATFORM_APPLE
-bool FVideoEncoderInputImpl::SetupForVulkan(VkInstance InApplicationVulkanInstance, VkPhysicalDevice InApplicationVulkanPhysicalDevice, VkDevice InApplicationVulkanDevice, uint32 InWidth, uint32 InHeight)
-{
-	FrameInfoVulkan.VulkanInstance = InApplicationVulkanInstance;
-	FrameInfoVulkan.VulkanPhysicalDevice = InApplicationVulkanPhysicalDevice;
-	FrameInfoVulkan.VulkanDevice = InApplicationVulkanDevice;
+	bool FVideoEncoderInputImpl::SetupForVulkan(VkInstance InApplicationVulkanInstance, VkPhysicalDevice InApplicationVulkanPhysicalDevice, VkDevice InApplicationVulkanDevice, uint32 InWidth, uint32 InHeight)
+	{
+		FrameInfoVulkan.VulkanInstance = InApplicationVulkanInstance;
+		FrameInfoVulkan.VulkanPhysicalDevice = InApplicationVulkanPhysicalDevice;
+		FrameInfoVulkan.VulkanDevice = InApplicationVulkanDevice;
 
-	FrameFormat = EVideoFrameFormat::VULKAN_R8G8B8A8_UNORM;
-	this->Width = InWidth;
-	this->Height = InHeight;
+		FrameFormat = EVideoFrameFormat::VULKAN_R8G8B8A8_UNORM;
+		this->Width = InWidth;
+		this->Height = InHeight;
 
-	CollectAvailableEncoders();
-	return true;
-}
+		CollectAvailableEncoders();
+		return true;
+	}
 #endif
 
-// --- available encoders -------------------------------------------------------------------------
+	// --- available encoders -------------------------------------------------------------------------
 
-void FVideoEncoderInputImpl::CollectAvailableEncoders()
-{
-	AvailableEncoders.Empty();
-	for (const FVideoEncoderInfo& Info : FVideoEncoderFactory::Get().GetAvailable())
+	void FVideoEncoderInputImpl::CollectAvailableEncoders()
 	{
-		if (Info.SupportedInputFormats.Contains(FrameFormat))
+		AvailableEncoders.Empty();
+		for (const FVideoEncoderInfo& Info : FVideoEncoderFactory::Get().GetAvailable())
 		{
-			AvailableEncoders.Push(Info);
+			if (Info.SupportedInputFormats.Contains(FrameFormat))
+			{
+				AvailableEncoders.Push(Info);
+			}
 		}
 	}
-}
 
-const TArray<FVideoEncoderInfo>& FVideoEncoderInputImpl::GetAvailableEncoders()
-{
-	return AvailableEncoders;
-}
+	const TArray<FVideoEncoderInfo>& FVideoEncoderInputImpl::GetAvailableEncoders()
+	{
+		return AvailableEncoders;
+	}
 
-// --- encoder input frames -----------------------------------------------------------------------
+	// --- encoder input frames -----------------------------------------------------------------------
 
-// create a user managed buffer
-FVideoEncoderInputFrame* FVideoEncoderInputImpl::CreateBuffer(OnFrameReleasedCallback InOnFrameReleased)
-{
-	FVideoEncoderInputFrameImpl* Frame = CreateFrame();
-	if (Frame)
+	bool FVideoEncoderInputImpl::IsUserManagedFrame(const FVideoEncoderInputFrame* InBuffer) const
 	{
+		const FVideoEncoderInputFrameImpl* Frame = static_cast<const FVideoEncoderInputFrameImpl*>(InBuffer);
 		FScopeLock						Guard(&ProtectFrames);
-		UserManagedFrames.Emplace(Frame, MoveTemp(InOnFrameReleased));
+		for (int32 Index = UserManagedFrames.Num() - 1; Index >= 0; --Index)
+		{
+			if (UserManagedFrames[Index].Key == Frame)
+			{
+				return true;
+			}
+		}
+		return false;
 	}
-	return Frame;
-}
 
-// destroy user managed buffer
-void FVideoEncoderInputImpl::DestroyBuffer(FVideoEncoderInputFrame* InBuffer)
-{
-	FVideoEncoderInputFrameImpl*	Frame = static_cast<FVideoEncoderInputFrameImpl*>(InBuffer);
-	FScopeLock						Guard(&ProtectFrames);
-	bool							bAnythingRemoved = false;
-	for (int32 Index = UserManagedFrames.Num() - 1; Index >= 0; --Index)
+	// create a user managed buffer
+	FVideoEncoderInputFrame* FVideoEncoderInputImpl::CreateBuffer(OnFrameReleasedCallback InOnFrameReleased)
 	{
-		if (UserManagedFrames[Index].Key == Frame)
+		FVideoEncoderInputFrameImpl* Frame = CreateFrame();
+		if (Frame)
 		{
-			UserManagedFrames.RemoveAt(Index);
-			bAnythingRemoved = true;
+			FScopeLock						Guard(&ProtectFrames);
+			UserManagedFrames.Emplace(Frame, MoveTemp(InOnFrameReleased));
 		}
+		return Frame;
 	}
-	if (bAnythingRemoved)
+
+	// destroy user managed buffer
+	void FVideoEncoderInputImpl::DestroyBuffer(FVideoEncoderInputFrame* InBuffer)
 	{
-		delete Frame;
+		FVideoEncoderInputFrameImpl* Frame = static_cast<FVideoEncoderInputFrameImpl*>(InBuffer);
+		FScopeLock						Guard(&ProtectFrames);
+		bool							bAnythingRemoved = false;
+		for (int32 Index = UserManagedFrames.Num() - 1; Index >= 0; --Index)
+		{
+			if (UserManagedFrames[Index].Key == Frame)
+			{
+				UserManagedFrames.RemoveAt(Index);
+				bAnythingRemoved = true;
+			}
+		}
+		if (bAnythingRemoved)
+		{
+			delete Frame;
+		}
 	}
-}
 
-// --- encoder input frames -----------------------------------------------------------------------
+	// --- encoder input frames -----------------------------------------------------------------------
 
-FVideoEncoderInputFrame* FVideoEncoderInputImpl::ObtainInputFrame()
-{
-	FVideoEncoderInputFrameImpl*	Frame = nullptr;
-	FScopeLock						Guard(&ProtectFrames);
-
-	if (!AvailableFrames.IsEmpty())
+	FVideoEncoderInputFrame* FVideoEncoderInputImpl::ObtainInputFrame()
 	{
+		FVideoEncoderInputFrameImpl* Frame = nullptr;
+		FScopeLock						Guard(&ProtectFrames);
 
-		AvailableFrames.Dequeue(Frame);
-		
-	}
-	else 
-	{
-		Frame = CreateFrame();
-		UE_LOG(LogVideoEncoder, Verbose, TEXT("Created new frame total frames: %d"), NumBuffers);
-	}
+		if (!AvailableFrames.IsEmpty())
+		{
 
-	ActiveFrames.Push(Frame);
-	
-	Frame->SetFrameID(NextFrameID++);
+			AvailableFrames.Dequeue(Frame);
 
-	if (NextFrameID == 0)
-	{
-		++NextFrameID; // skip 0 id
-	} 
+		}
+		else
+		{
+			Frame = CreateFrame();
+			UE_LOG(LogVideoEncoder, Verbose, TEXT("Created new frame total frames: %d"), NumBuffers);
+		}
 
-	return const_cast<FVideoEncoderInputFrame*>(Frame->Obtain());
-}
+		ActiveFrames.Push(Frame);
 
-FVideoEncoderInputFrameImpl* FVideoEncoderInputImpl::CreateFrame()
-{
-	FVideoEncoderInputFrameImpl*	Frame = new FVideoEncoderInputFrameImpl(this);
-	NumBuffers++;
-	switch (FrameFormat)
-	{
-	case EVideoFrameFormat::Undefined:
-		UE_LOG(LogVideoEncoder, Error, TEXT("Got undefined frame format!"))
-		break;
-	case EVideoFrameFormat::YUV420P:
-		SetupFrameYUV420P(Frame);
-		break;
-	case EVideoFrameFormat::D3D11_R8G8B8A8_UNORM:
-		SetupFrameD3D11(Frame);
-		break;
-	case EVideoFrameFormat::D3D12_R8G8B8A8_UNORM:
-		SetupFrameD3D12(Frame);
-		break;
-	case EVideoFrameFormat::VULKAN_R8G8B8A8_UNORM:
-		SetupFrameVulkan(Frame);
-		break;
-	case EVideoFrameFormat::CUDA_R8G8B8A8_UNORM:
-		SetupFrameCUDA(Frame);
-		break;
-	default:
-		check(false);
-		break;
-	}
-	return Frame;
-}
+		Frame->SetFrameID(NextFrameID++);
 
-void FVideoEncoderInputImpl::ReleaseInputFrame(FVideoEncoderInputFrame* InFrame)
-{
-	FVideoEncoderInputFrameImpl* InFrameImpl = static_cast<FVideoEncoderInputFrameImpl*>(InFrame);
+		if (NextFrameID == 0)
+		{
+			++NextFrameID; // skip 0 id
+		}
+
+		return const_cast<FVideoEncoderInputFrame*>(Frame->Obtain());
+	}
 
-	FScopeLock	Guard(&ProtectFrames);
-	// check user managed buffers first
-	for (const UserManagedFrame& Frame : UserManagedFrames)
+	FVideoEncoderInputFrameImpl* FVideoEncoderInputImpl::CreateFrame()
 	{
-		if (Frame.Key == InFrameImpl)
+		FVideoEncoderInputFrameImpl* Frame = new FVideoEncoderInputFrameImpl(this);
+		NumBuffers++;
+		switch (FrameFormat)
 		{
-			Frame.Value(InFrameImpl);
-			return;
+		case EVideoFrameFormat::Undefined:
+			UE_LOG(LogVideoEncoder, Error, TEXT("Got undefined frame format!"));
+			break;
+		case EVideoFrameFormat::YUV420P:
+			SetupFrameYUV420P(Frame);
+			break;
+		case EVideoFrameFormat::D3D11_R8G8B8A8_UNORM:
+			SetupFrameD3D11(Frame);
+			break;
+		case EVideoFrameFormat::D3D12_R8G8B8A8_UNORM:
+			SetupFrameD3D12(Frame);
+			break;
+		case EVideoFrameFormat::VULKAN_R8G8B8A8_UNORM:
+			SetupFrameVulkan(Frame);
+			break;
+		case EVideoFrameFormat::CUDA_R8G8B8A8_UNORM:
+			SetupFrameCUDA(Frame);
+			break;
+		default:
+			check(false);
+			break;
 		}
+		return Frame;
 	}
 
-	int32		NumRemoved = ActiveFrames.Remove(InFrameImpl);
-	check(NumRemoved == 1);
-	if (NumRemoved > 0)
+	void FVideoEncoderInputImpl::ReleaseInputFrame(FVideoEncoderInputFrame* InFrame)
 	{
-		// drop frame if format changed
-		if (InFrame->GetFormat() != FrameFormat)
+		FVideoEncoderInputFrameImpl* InFrameImpl = static_cast<FVideoEncoderInputFrameImpl*>(InFrame);
+
+		FScopeLock Guard(&ProtectFrames);
+		// check user managed buffers first
+		for (const UserManagedFrame& Frame : UserManagedFrames)
 		{
-			ProtectFrames.Unlock();
-			delete InFrameImpl;
-			NumBuffers--;
-			UE_LOG(LogVideoEncoder, Verbose, TEXT("Deleted buffer (format mismatch) total remaining: %d"), NumBuffers);
-			return;
+			if (Frame.Key == InFrameImpl)
+			{
+				Frame.Value(InFrameImpl);
+				return;
+			}
 		}
-		
-		// drop frame if resolution changed
-		if(bIsResizable && (InFrame->GetWidth() != this->Width || InFrame->GetHeight() != this->Height))
+
+		int32 NumRemoved = ActiveFrames.Remove(InFrameImpl);
+		check(NumRemoved == 1);
+		if (NumRemoved > 0)
 		{
-			ProtectFrames.Unlock();
-			delete InFrameImpl;
-			NumBuffers--;
-			UE_LOG(LogVideoEncoder, Verbose, TEXT("Deleted buffer (size mismatch) total remaining: %d"), NumBuffers);
-			return;
+			// drop frame if format changed
+			if (InFrame->GetFormat() != FrameFormat)
+			{
+				ProtectFrames.Unlock();
+				delete InFrameImpl;
+				NumBuffers--;
+				UE_LOG(LogVideoEncoder, Verbose, TEXT("Deleted buffer (format mismatch) total remaining: %d"), NumBuffers);
+				return;
+			}
+
+			// drop frame if resolution changed
+			if (bIsResizable && (InFrame->GetWidth() != this->Width || InFrame->GetHeight() != this->Height))
+			{
+				ProtectFrames.Unlock();
+				delete InFrameImpl;
+				NumBuffers--;
+				UE_LOG(LogVideoEncoder, Verbose, TEXT("Deleted buffer (size mismatch) total remaining: %d"), NumBuffers);
+				return;
+			}
+
+			if (!AvailableFrames.IsEmpty() && NumBuffers > MaxNumBuffers)
+			{
+				ProtectFrames.Unlock();
+				delete InFrameImpl;
+				NumBuffers--;
+				UE_LOG(LogVideoEncoder, Verbose, TEXT("Deleted buffer (too many) total frames: %d"), NumBuffers);
+				return;
+			}
+
+			AvailableFrames.Enqueue(InFrameImpl);
 		}
+	}
 
-		if(!AvailableFrames.IsEmpty() && NumBuffers > MaxNumBuffers)
+	void FVideoEncoderInputImpl::Flush()
+	{
+		ProtectFrames.Lock();
+		while (!AvailableFrames.IsEmpty())
 		{
+			FVideoEncoderInputFrameImpl* Frame = nullptr;
+			AvailableFrames.Dequeue(Frame);
 			ProtectFrames.Unlock();
-			delete InFrameImpl;
+			delete Frame;
 			NumBuffers--;
-			UE_LOG(LogVideoEncoder, Verbose, TEXT("Deleted buffer (too many) total frames: %d"), NumBuffers);
-			return;
+			ProtectFrames.Lock();
 		}
-
-		AvailableFrames.Enqueue(InFrameImpl);
+		ProtectFrames.Unlock();
 	}
-}
 
-void FVideoEncoderInputImpl::Flush()
-{
-	ProtectFrames.Lock();
-	while (!AvailableFrames.IsEmpty())
+	void FVideoEncoderInputImpl::SetupFrameYUV420P(FVideoEncoderInputFrameImpl* Frame)
 	{
-		FVideoEncoderInputFrameImpl*	Frame = nullptr;
-		AvailableFrames.Dequeue(Frame);
-		ProtectFrames.Unlock();
-		delete Frame;
-		NumBuffers--;
-		ProtectFrames.Lock();
+		Frame->SetFormat(EVideoFrameFormat::YUV420P);
+		Frame->SetWidth(this->Width);
+		Frame->SetHeight(this->Height);
+		FVideoEncoderInputFrame::FYUV420P& YUV420P = Frame->GetYUV420P();
+		YUV420P.StrideY = FrameInfoYUV420P.StrideY;
+		YUV420P.StrideU = FrameInfoYUV420P.StrideU;
+		YUV420P.StrideV = FrameInfoYUV420P.StrideV;
+		YUV420P.Data[0] = YUV420P.Data[1] = YUV420P.Data[2] = nullptr;
 	}
-	ProtectFrames.Unlock();
-}
 
-void FVideoEncoderInputImpl::SetupFrameYUV420P(FVideoEncoderInputFrameImpl* Frame)
-{
-	Frame->SetFormat(EVideoFrameFormat::YUV420P);
-	Frame->SetWidth(this->Width);
-	Frame->SetHeight(this->Height);
-	FVideoEncoderInputFrame::FYUV420P& YUV420P = Frame->GetYUV420P();
-	YUV420P.StrideY = FrameInfoYUV420P.StrideY;
-	YUV420P.StrideU = FrameInfoYUV420P.StrideU;
-	YUV420P.StrideV = FrameInfoYUV420P.StrideV;
-	YUV420P.Data[0] = YUV420P.Data[1] = YUV420P.Data[2] = nullptr;
-}
-
-void FVideoEncoderInputImpl::SetupFrameD3D11(FVideoEncoderInputFrameImpl* Frame)
-{
-#if PLATFORM_WINDOWS
-	Frame->SetFormat(FrameFormat);
-	Frame->SetWidth(this->Width);
-	Frame->SetHeight(this->Height);
-
-	FVideoEncoderInputFrame::FD3D11& Data = Frame->GetD3D11();
-	Data.EncoderDevice = FrameInfoD3D.EncoderDeviceD3D11;
-#endif
-}
-
-void FVideoEncoderInputImpl::SetupFrameD3D12(FVideoEncoderInputFrameImpl* Frame)
-{
+	void FVideoEncoderInputImpl::SetupFrameD3D11(FVideoEncoderInputFrameImpl* Frame)
+	{
 #if PLATFORM_WINDOWS
-	Frame->SetFormat(FrameFormat);
-	Frame->SetWidth(this->Width);
-	Frame->SetHeight(this->Height);
+		Frame->SetFormat(FrameFormat);
+		Frame->SetWidth(this->Width);
+		Frame->SetHeight(this->Height);
 
-	if (FrameFormat == EVideoFrameFormat::D3D11_R8G8B8A8_UNORM)
-	{
 		FVideoEncoderInputFrame::FD3D11& Data = Frame->GetD3D11();
 		Data.EncoderDevice = FrameInfoD3D.EncoderDeviceD3D11;
-	}
-	else
-	{	
-		FVideoEncoderInputFrame::FD3D12& Data = Frame->GetD3D12();
-		Data.EncoderDevice = FrameInfoD3D.EncoderDeviceD3D12;
-	}
 #endif
-}
-
-void FVideoEncoderInputImpl::SetupFrameVulkan(FVideoEncoderInputFrameImpl* Frame)
-{
-#if PLATFORM_DESKTOP && !PLATFORM_APPLE
-	Frame->SetFormat(FrameFormat);
-	Frame->SetWidth(this->Width);
-	Frame->SetHeight(this->Height);
-
-	FVideoEncoderInputFrame::FVulkan& Data = Frame->GetVulkan();
-	Data.EncoderDevice = FrameInfoVulkan.VulkanDevice;
-#endif
-}
-
-void FVideoEncoderInputImpl::SetupFrameCUDA(FVideoEncoderInputFrameImpl* Frame)
-{
-	Frame->SetFormat(FrameFormat);
-	Frame->SetWidth(this->Width);
-	Frame->SetHeight(this->Height);
-	FVideoEncoderInputFrame::FCUDA& Data = Frame->GetCUDA();
-	Data.EncoderDevice = FrameInfoCUDA.EncoderContextCUDA;
-}
-
-// ---
+	}
 
+	void FVideoEncoderInputImpl::SetupFrameD3D12(FVideoEncoderInputFrameImpl* Frame)
+	{
 #if PLATFORM_WINDOWS
-TRefCountPtr<ID3D11Device> FVideoEncoderInputImpl::GetD3D11EncoderDevice() const
-{
-	return FrameInfoD3D.EncoderDeviceD3D11;
-}
+		Frame->SetFormat(FrameFormat);
+		Frame->SetWidth(this->Width);
+		Frame->SetHeight(this->Height);
 
-TRefCountPtr<ID3D12Device> FVideoEncoderInputImpl::GetD3D12EncoderDevice() const
-{		
-	return FrameInfoD3D.EncoderDeviceD3D12;
-}
+		if (FrameFormat == EVideoFrameFormat::D3D11_R8G8B8A8_UNORM)
+		{
+			FVideoEncoderInputFrame::FD3D11& Data = Frame->GetD3D11();
+			Data.EncoderDevice = FrameInfoD3D.EncoderDeviceD3D11;
+		}
+		else
+		{
+			FVideoEncoderInputFrame::FD3D12& Data = Frame->GetD3D12();
+			Data.EncoderDevice = FrameInfoD3D.EncoderDeviceD3D12;
+		}
 #endif
+	}
 
-CUcontext FVideoEncoderInputImpl::GetCUDAEncoderContext() const
-{
-	return FrameInfoCUDA.EncoderContextCUDA;
-}
-
+	void FVideoEncoderInputImpl::SetupFrameVulkan(FVideoEncoderInputFrameImpl* Frame)
+	{
 #if PLATFORM_DESKTOP && !PLATFORM_APPLE
-void* FVideoEncoderInputImpl::GetVulkanEncoderDevice() const
-{
-	return (void*)&FrameInfoVulkan;
-}
-#endif
+		Frame->SetFormat(FrameFormat);
+		Frame->SetWidth(this->Width);
+		Frame->SetHeight(this->Height);
 
-// *** FVideoEncoderInputFrame ********************************************************************
-
-FVideoEncoderInputFrame::FVideoEncoderInputFrame()
-	: FrameID(0)
-	, TimestampUs(0)
-	, TimestampRTP(0)
-	, NumReferences(0)
-	, Format(EVideoFrameFormat::Undefined)
-	, Width(0)
-	, Height(0)
-	, bFreeYUV420PData(false)
-{
-}
-
-FVideoEncoderInputFrame::FVideoEncoderInputFrame(const FVideoEncoderInputFrame& CloneFrom)
-	: FrameID(CloneFrom.FrameID)
-	, TimestampUs(CloneFrom.TimestampUs)
-	, TimestampRTP(CloneFrom.TimestampRTP)
-	, NumReferences(0)
-	, Format(CloneFrom.Format)
-	, Width(CloneFrom.Width)
-	, Height(CloneFrom.Height)
-	, bFreeYUV420PData(false)
-{
-#if PLATFORM_WINDOWS
-	D3D11.EncoderDevice = CloneFrom.D3D11.EncoderDevice;
-	D3D11.Texture = CloneFrom.D3D11.Texture;
-	if ((D3D11.EncoderTexture = CloneFrom.D3D11.EncoderTexture) != nullptr)
-	{
-		D3D11.EncoderTexture->AddRef();
+		FVideoEncoderInputFrame::FVulkan& Data = Frame->GetVulkan();
+		Data.EncoderDevice = FrameInfoVulkan.VulkanDevice;
+#endif
 	}
 
-	D3D12.EncoderDevice = CloneFrom.D3D12.EncoderDevice;
-	D3D12.Texture = CloneFrom.D3D12.Texture;
-	if ((D3D12.EncoderTexture = CloneFrom.D3D12.EncoderTexture) != nullptr)
+	void FVideoEncoderInputImpl::SetupFrameCUDA(FVideoEncoderInputFrameImpl* Frame)
 	{
-		D3D12.EncoderTexture->AddRef();
+		Frame->SetFormat(FrameFormat);
+		Frame->SetWidth(this->Width);
+		Frame->SetHeight(this->Height);
+		FVideoEncoderInputFrame::FCUDA& Data = Frame->GetCUDA();
+		Data.EncoderDevice = FrameInfoCUDA.EncoderContextCUDA;
 	}
-#endif
-
-	CUDA.EncoderDevice = CloneFrom.CUDA.EncoderDevice;
-	CUDA.EncoderTexture = CloneFrom.CUDA.EncoderTexture;
-}
 
-FVideoEncoderInputFrame::~FVideoEncoderInputFrame()
-{
-	if (bFreeYUV420PData)
-	{
-		delete[] YUV420P.Data[0];
-		delete[] YUV420P.Data[1];
-		delete[] YUV420P.Data[2];
-		YUV420P.Data[0] = YUV420P.Data[1] = YUV420P.Data[2] = nullptr;
-		bFreeYUV420PData = false;
-	}
+	// ---
 
 #if PLATFORM_WINDOWS
-	if (D3D11.EncoderTexture)
-	{
-		// check to make sure this frame holds the last reference
-		auto	NumRef = D3D11.EncoderTexture->AddRef();
-		if (NumRef > 2)
-		{
-			UE_LOG(LogVideoEncoder, Warning, TEXT("VideoEncoderFame - D3D11 input texture still holds %d references."), NumRef);
-		}
-		D3D11.EncoderTexture->Release();
-		D3D11.EncoderTexture->Release();
-		D3D11.EncoderTexture = nullptr;
-	}
-	if (D3D11.SharedHandle)
-	{
-		CloseHandle(D3D11.SharedHandle);
-		D3D11.SharedHandle = nullptr;
-	}
-	if (D3D11.Texture && OnReleaseD3D11Texture)
-	{
-		OnReleaseD3D11Texture(D3D11.Texture);
-	}
-	if (D3D12.EncoderTexture)
+	TRefCountPtr<ID3D11Device> FVideoEncoderInputImpl::GetD3D11EncoderDevice() const
 	{
-		// check to make sure this frame holds the last reference
-		auto	NumRef = D3D12.EncoderTexture->AddRef();
-		if (NumRef > 2)
-		{
-			UE_LOG(LogVideoEncoder, Warning, TEXT("VideoEncoderFame - D3D12 input texture still holds %d references."), NumRef);
-		}
-		D3D12.EncoderTexture->Release();
-		D3D12.EncoderTexture->Release();
-		D3D12.EncoderTexture = nullptr;
+		return FrameInfoD3D.EncoderDeviceD3D11;
 	}
-	if (D3D12.Texture && OnReleaseD3D12Texture)
+
+	TRefCountPtr<ID3D12Device> FVideoEncoderInputImpl::GetD3D12EncoderDevice() const
 	{
-		OnReleaseD3D12Texture(D3D12.Texture);
-		D3D12.Texture = nullptr;
+		return FrameInfoD3D.EncoderDeviceD3D12;
 	}
-#endif 
+#endif
 
-	if (CUDA.EncoderTexture)
+	CUcontext FVideoEncoderInputImpl::GetCUDAEncoderContext() const
 	{
-		OnReleaseCUDATexture(CUDA.EncoderTexture);
-		CUDA.EncoderTexture = nullptr;
+		return FrameInfoCUDA.EncoderContextCUDA;
 	}
 
 #if PLATFORM_DESKTOP && !PLATFORM_APPLE
-	if (Vulkan.EncoderTexture != VK_NULL_HANDLE)
+	void* FVideoEncoderInputImpl::GetVulkanEncoderDevice() const
 	{
-		OnReleaseVulkanTexture(Vulkan.EncoderTexture);
+		return (void*)&FrameInfoVulkan;
 	}
+#endif
 
-	if (Vulkan.EncoderSurface)
+	// *** FVideoEncoderInputFrame ********************************************************************
+
+	FVideoEncoderInputFrame::FVideoEncoderInputFrame()
+		: FrameID(0)
+		, TimestampUs(0)
+		, TimestampRTP(0)
+		, NumReferences(0)
+		, Format(EVideoFrameFormat::Undefined)
+		, Width(0)
+		, Height(0)
+		, bFreeYUV420PData(false)
 	{
-		OnReleaseVulkanSurface(Vulkan.EncoderSurface);
 	}
-#endif
-}
 
-void FVideoEncoderInputFrame::AllocateYUV420P()
-{
-	if (!bFreeYUV420PData)
+	FVideoEncoderInputFrame::FVideoEncoderInputFrame(const FVideoEncoderInputFrame& CloneFrom)
+		: FrameID(CloneFrom.FrameID)
+		, TimestampUs(CloneFrom.TimestampUs)
+		, TimestampRTP(CloneFrom.TimestampRTP)
+		, NumReferences(0)
+		, Format(CloneFrom.Format)
+		, Width(CloneFrom.Width)
+		, Height(CloneFrom.Height)
+		, bFreeYUV420PData(false)
 	{
-		YUV420P.StrideY = Width;
-		YUV420P.StrideU = (Width + 1) / 2;
-		YUV420P.StrideV = (Width + 1) / 2;;
-		YUV420P.Data[0] = new uint8[Height * YUV420P.StrideY];
-		YUV420P.Data[1] = new uint8[(Height + 1) / 2 * YUV420P.StrideU];
-		YUV420P.Data[2] = new uint8[(Height + 1) / 2 * YUV420P.StrideV];
-		bFreeYUV420PData = true;
+#if PLATFORM_WINDOWS
+		D3D11.EncoderDevice = CloneFrom.D3D11.EncoderDevice;
+		D3D11.Texture = CloneFrom.D3D11.Texture;
+		if ((D3D11.EncoderTexture = CloneFrom.D3D11.EncoderTexture) != nullptr)
+		{
+			D3D11.EncoderTexture->AddRef();
+		}
+
+		D3D12.EncoderDevice = CloneFrom.D3D12.EncoderDevice;
+		D3D12.Texture = CloneFrom.D3D12.Texture;
+		if ((D3D12.EncoderTexture = CloneFrom.D3D12.EncoderTexture) != nullptr)
+		{
+			D3D12.EncoderTexture->AddRef();
+		}
+#endif
+
+		CUDA.EncoderDevice = CloneFrom.CUDA.EncoderDevice;
+		CUDA.EncoderTexture = CloneFrom.CUDA.EncoderTexture;
 	}
-}
 
-void FVideoEncoderInputFrame::SetYUV420P(const uint8* InDataY, const uint8* InDataU, const uint8* InDataV, uint32 InStrideY, uint32 InStrideU, uint32 InStrideV)
-{
-	if (Format == EVideoFrameFormat::YUV420P)
+	FVideoEncoderInputFrame::~FVideoEncoderInputFrame()
 	{
 		if (bFreeYUV420PData)
 		{
 			delete[] YUV420P.Data[0];
 			delete[] YUV420P.Data[1];
 			delete[] YUV420P.Data[2];
+			YUV420P.Data[0] = YUV420P.Data[1] = YUV420P.Data[2] = nullptr;
 			bFreeYUV420PData = false;
 		}
-		YUV420P.Data[0] = InDataY;
-		YUV420P.Data[1] = InDataU;
-		YUV420P.Data[2] = InDataV;
-		YUV420P.StrideY = InStrideY;
-		YUV420P.StrideU = InStrideU;
-		YUV420P.StrideV = InStrideV;
-	}
-}
-
-static FThreadSafeCounter	_VideoEncoderInputFrameCnt{ 0 };
 
 #if PLATFORM_WINDOWS
-void FVideoEncoderInputFrame::SetTexture(ID3D11Texture2D* InTexture, FReleaseD3D11TextureCallback InOnReleaseTexture)
-{
-	if (Format == EVideoFrameFormat::D3D11_R8G8B8A8_UNORM)
-	{
-		check(D3D11.Texture == nullptr);
-		if (!D3D11.Texture)
+		if (D3D11.EncoderTexture)
 		{
-			TRefCountPtr<IDXGIResource> DXGIResource;
-			HANDLE						SharedHandle;
-			HRESULT		Result = InTexture->QueryInterface(IID_PPV_ARGS(DXGIResource.GetInitReference()));
-			if (FAILED(Result))
+			// check to make sure this frame holds the last reference
+			auto	NumRef = D3D11.EncoderTexture->AddRef();
+			if (NumRef > 2)
 			{
-				UE_LOG(LogVideoEncoder, Error, TEXT("ID3D11Device::QueryInterface() failed 0x%X - %s."), Result, *GetComErrorDescription(Result));
+				UE_LOG(LogVideoEncoder, Warning, TEXT("VideoEncoderFame - D3D11 input texture still holds %d references."), NumRef);
 			}
-			//
-			// NOTE : The HANDLE IDXGIResource::GetSharedHandle gives us is NOT an NT Handle, and therefre we should not call CloseHandle on it
-			//
-			else if((Result = DXGIResource->GetSharedHandle(&SharedHandle)) != S_OK)
-			{
-				UE_LOG(LogVideoEncoder, Error, TEXT("IDXGIResource::GetSharedHandle() failed 0x%X - %s."), Result, *GetComErrorDescription(Result));
-			}
-			else if (SharedHandle == nullptr)
-			{
-				UE_LOG(LogVideoEncoder, Error, TEXT("IDXGIResource::GetSharedHandle() failed to return a shared texture resource no created as shared? (D3D11_RESOURCE_MISC_SHARED)."));
-			}
-			else if((Result = D3D11.EncoderDevice->OpenSharedResource(SharedHandle, __uuidof(ID3D11Texture2D), (LPVOID*)&D3D11.EncoderTexture)) != S_OK)
-			{
-				UE_LOG(LogVideoEncoder, Error, TEXT("ID3D11Device::OpenSharedResource() failed 0x%X - %s."), Result, *GetComErrorDescription(Result));
-			}
-			else
+			D3D11.EncoderTexture->Release();
+			D3D11.EncoderTexture->Release();
+			D3D11.EncoderTexture = nullptr;
+		}
+		if (D3D11.SharedHandle)
+		{
+			CloseHandle(D3D11.SharedHandle);
+			D3D11.SharedHandle = nullptr;
+		}
+		if (D3D11.Texture && OnReleaseD3D11Texture)
+		{
+			OnReleaseD3D11Texture(D3D11.Texture);
+		}
+		if (D3D12.EncoderTexture)
+		{
+			// check to make sure this frame holds the last reference
+			auto	NumRef = D3D12.EncoderTexture->AddRef();
+			if (NumRef > 2)
 			{
-				DEBUG_SET_D3D11_OBJECT_NAME(D3D11.EncoderTexture, "FVideoEncoderInputFrame::SetTexture()");
-				D3D11.Texture = InTexture;
-				OnReleaseD3D11Texture = InOnReleaseTexture;
+				UE_LOG(LogVideoEncoder, Warning, TEXT("VideoEncoderFame - D3D12 input texture still holds %d references."), NumRef);
 			}
+			D3D12.EncoderTexture->Release();
+			D3D12.EncoderTexture->Release();
+			D3D12.EncoderTexture = nullptr;
+		}
+		if (D3D12.Texture && OnReleaseD3D12Texture)
+		{
+			OnReleaseD3D12Texture(D3D12.Texture);
+			D3D12.Texture = nullptr;
 		}
+#endif 
+
+		if (CUDA.EncoderTexture)
+		{
+			OnReleaseCUDATexture(CUDA.EncoderTexture);
+			CUDA.EncoderTexture = nullptr;
+		}
+
+#if PLATFORM_DESKTOP && !PLATFORM_APPLE
+		if (Vulkan.EncoderTexture != VK_NULL_HANDLE)
+		{
+			OnReleaseVulkanTexture(Vulkan.EncoderTexture);
+		}
+
+		if (Vulkan.EncoderSurface)
+		{
+			OnReleaseVulkanSurface(Vulkan.EncoderSurface);
+		}
+#endif
 	}
-}
 
-void FVideoEncoderInputFrame::SetTexture(ID3D12Resource* InTexture, FReleaseD3D12TextureCallback InOnReleaseTexture)
-{
-	if (Format == EVideoFrameFormat::D3D11_R8G8B8A8_UNORM)
+	void FVideoEncoderInputFrame::AllocateYUV420P()
 	{
-		check(D3D12.Texture == nullptr)
-		check(D3D11.EncoderTexture == nullptr)
-
-		TRefCountPtr<ID3D12Device>	OwnerDevice;
-		HRESULT						Result;
-		if((Result = InTexture->GetDevice(IID_PPV_ARGS(OwnerDevice.GetInitReference()))) != S_OK)
+		if (!bFreeYUV420PData)
 		{
-			UE_LOG(LogVideoEncoder, Error, TEXT("ID3D11Device::QueryInterface() failed 0x%X - %s."), Result, *GetComErrorDescription(Result));
+			YUV420P.StrideY = Width;
+			YUV420P.StrideU = (Width + 1) / 2;
+			YUV420P.StrideV = (Width + 1) / 2;;
+			YUV420P.Data[0] = new uint8[Height * YUV420P.StrideY];
+			YUV420P.Data[1] = new uint8[(Height + 1) / 2 * YUV420P.StrideU];
+			YUV420P.Data[2] = new uint8[(Height + 1) / 2 * YUV420P.StrideV];
+			bFreeYUV420PData = true;
 		}
-		//
-		// NOTE: ID3D12Device::CreateSharedHandle gives as an NT Handle, and so we need to call CloseHandle on it
-		//
-		else if ((Result = OwnerDevice->CreateSharedHandle(InTexture, NULL, GENERIC_ALL, *FString::Printf(TEXT("%s_FVideoEncoderInputFrame_%d"), *GetGUID(), _VideoEncoderInputFrameCnt.Increment()), &D3D11.SharedHandle)) != S_OK)
+	}
+
+	void FVideoEncoderInputFrame::SetYUV420P(const uint8* InDataY, const uint8* InDataU, const uint8* InDataV, uint32 InStrideY, uint32 InStrideU, uint32 InStrideV)
+	{
+		if (Format == EVideoFrameFormat::YUV420P)
 		{
-			UE_LOG(LogVideoEncoder, Error, TEXT("ID3D12Device::CreateSharedHandle() failed 0x%X - %s."), Result, *GetComErrorDescription(Result));
+			if (bFreeYUV420PData)
+			{
+				delete[] YUV420P.Data[0];
+				delete[] YUV420P.Data[1];
+				delete[] YUV420P.Data[2];
+				bFreeYUV420PData = false;
+			}
+			YUV420P.Data[0] = InDataY;
+			YUV420P.Data[1] = InDataU;
+			YUV420P.Data[2] = InDataV;
+			YUV420P.StrideY = InStrideY;
+			YUV420P.StrideU = InStrideU;
+			YUV420P.StrideV = InStrideV;
 		}
-		else if (D3D11.SharedHandle == nullptr)
+	}
+
+	static FThreadSafeCounter	_VideoEncoderInputFrameCnt{ 0 };
+
+#if PLATFORM_WINDOWS
+	void FVideoEncoderInputFrame::SetTexture(ID3D11Texture2D* InTexture, FReleaseD3D11TextureCallback InOnReleaseTexture)
+	{
+		if (Format == EVideoFrameFormat::D3D11_R8G8B8A8_UNORM)
 		{
-			UE_LOG(LogVideoEncoder, Error, TEXT("ID3D12Device::CreateSharedHandle() failed to return a shared texture resource no created as shared? (D3D12_RESOURCE_FLAG_ALLOW_SIMULTANEOUS_ACCESS)."));
+			check(D3D11.Texture == nullptr);
+			if (!D3D11.Texture)
+			{
+				TRefCountPtr<IDXGIResource> DXGIResource;
+				HANDLE						SharedHandle;
+				HRESULT		Result = InTexture->QueryInterface(IID_PPV_ARGS(DXGIResource.GetInitReference()));
+				if (FAILED(Result))
+				{
+					UE_LOG(LogVideoEncoder, Error, TEXT("ID3D11Device::QueryInterface() failed 0x%X - %s."), Result, *GetComErrorDescription(Result));
+				}
+				//
+				// NOTE : The HANDLE IDXGIResource::GetSharedHandle gives us is NOT an NT Handle, and therefre we should not call CloseHandle on it
+				//
+				else if ((Result = DXGIResource->GetSharedHandle(&SharedHandle)) != S_OK)
+				{
+					UE_LOG(LogVideoEncoder, Error, TEXT("IDXGIResource::GetSharedHandle() failed 0x%X - %s."), Result, *GetComErrorDescription(Result));
+				}
+				else if (SharedHandle == nullptr)
+				{
+					UE_LOG(LogVideoEncoder, Error, TEXT("IDXGIResource::GetSharedHandle() failed to return a shared texture resource no created as shared? (D3D11_RESOURCE_MISC_SHARED)."));
+				}
+				else if ((Result = D3D11.EncoderDevice->OpenSharedResource(SharedHandle, __uuidof(ID3D11Texture2D), (LPVOID*)&D3D11.EncoderTexture)) != S_OK)
+				{
+					UE_LOG(LogVideoEncoder, Error, TEXT("ID3D11Device::OpenSharedResource() failed 0x%X - %s."), Result, *GetComErrorDescription(Result));
+				}
+				else
+				{
+					DEBUG_SET_D3D11_OBJECT_NAME(D3D11.EncoderTexture, "FVideoEncoderInputFrame::SetTexture()");
+					D3D11.Texture = InTexture;
+					OnReleaseD3D11Texture = InOnReleaseTexture;
+				}
+			}
 		}
-		else
+	}
+
+	void FVideoEncoderInputFrame::SetTexture(ID3D12Resource* InTexture, FReleaseD3D12TextureCallback InOnReleaseTexture)
+	{
+		if (Format == EVideoFrameFormat::D3D11_R8G8B8A8_UNORM)
 		{
-			TRefCountPtr <ID3D11Device1> Device1;
-			if ((Result = D3D11.EncoderDevice->QueryInterface(IID_PPV_ARGS(Device1.GetInitReference()))) != S_OK)
+			check(D3D12.Texture == nullptr)
+				check(D3D11.EncoderTexture == nullptr)
+
+				TRefCountPtr<ID3D12Device>	OwnerDevice;
+			HRESULT						Result;
+			if ((Result = InTexture->GetDevice(IID_PPV_ARGS(OwnerDevice.GetInitReference()))) != S_OK)
 			{
 				UE_LOG(LogVideoEncoder, Error, TEXT("ID3D11Device::QueryInterface() failed 0x%X - %s."), Result, *GetComErrorDescription(Result));
 			}
-			else if ((Result = Device1->OpenSharedResource1(D3D11.SharedHandle, IID_PPV_ARGS(&D3D11.EncoderTexture))) != S_OK)
+			//
+			// NOTE: ID3D12Device::CreateSharedHandle gives as an NT Handle, and so we need to call CloseHandle on it
+			//
+			else if ((Result = OwnerDevice->CreateSharedHandle(InTexture, NULL, GENERIC_ALL, *FString::Printf(TEXT("%s_FVideoEncoderInputFrame_%d"), *GetGUID(), _VideoEncoderInputFrameCnt.Increment()), &D3D11.SharedHandle)) != S_OK)
 			{
-				UE_LOG(LogVideoEncoder, Error, TEXT("ID3D11Device::OpenSharedResource1() failed 0x%X - %s."), Result, *GetComErrorDescription(Result));
+				UE_LOG(LogVideoEncoder, Error, TEXT("ID3D12Device::CreateSharedHandle() failed 0x%X - %s."), Result, *GetComErrorDescription(Result));
+			}
+			else if (D3D11.SharedHandle == nullptr)
+			{
+				UE_LOG(LogVideoEncoder, Error, TEXT("ID3D12Device::CreateSharedHandle() failed to return a shared texture resource no created as shared? (D3D12_RESOURCE_FLAG_ALLOW_SIMULTANEOUS_ACCESS)."));
 			}
 			else
 			{
-				DEBUG_SET_D3D11_OBJECT_NAME(D3D11.EncoderTexture, "FVideoEncoderInputFrame::SetTexture()");
-				D3D12.Texture = InTexture;
-				OnReleaseD3D12Texture = InOnReleaseTexture;
+				TRefCountPtr <ID3D11Device1> Device1;
+				if ((Result = D3D11.EncoderDevice->QueryInterface(IID_PPV_ARGS(Device1.GetInitReference()))) != S_OK)
+				{
+					UE_LOG(LogVideoEncoder, Error, TEXT("ID3D11Device::QueryInterface() failed 0x%X - %s."), Result, *GetComErrorDescription(Result));
+				}
+				else if ((Result = Device1->OpenSharedResource1(D3D11.SharedHandle, IID_PPV_ARGS(&D3D11.EncoderTexture))) != S_OK)
+				{
+					UE_LOG(LogVideoEncoder, Error, TEXT("ID3D11Device::OpenSharedResource1() failed 0x%X - %s."), Result, *GetComErrorDescription(Result));
+				}
+				else
+				{
+					DEBUG_SET_D3D11_OBJECT_NAME(D3D11.EncoderTexture, "FVideoEncoderInputFrame::SetTexture()");
+					D3D12.Texture = InTexture;
+					OnReleaseD3D12Texture = InOnReleaseTexture;
+				}
 			}
 		}
+		else
+		{
+			check(D3D12.Texture == nullptr);
+			check(D3D12.EncoderTexture == nullptr);
+			D3D12.Texture = InTexture;
+			D3D12.EncoderTexture = InTexture;
+			OnReleaseD3D12Texture = InOnReleaseTexture;
+		}
 	}
-	else
-	{
-		check(D3D12.Texture == nullptr);
-		check(D3D12.EncoderTexture == nullptr);
-		D3D12.Texture = InTexture;
-		D3D12.EncoderTexture = InTexture;		
-		OnReleaseD3D12Texture = InOnReleaseTexture;
-	}
-}
 #endif
 
-void FVideoEncoderInputFrame::SetTexture(CUarray InTexture, FReleaseCUDATextureCallback InOnReleaseTexture)
-{
-	if (Format == EVideoFrameFormat::CUDA_R8G8B8A8_UNORM)
+	void FVideoEncoderInputFrame::SetTexture(CUarray InTexture, FReleaseCUDATextureCallback InOnReleaseTexture)
 	{
-		CUDA.EncoderTexture = InTexture;
-		OnReleaseCUDATexture = InOnReleaseTexture;
-		if (!CUDA.EncoderTexture)
+		if (Format == EVideoFrameFormat::CUDA_R8G8B8A8_UNORM)
 		{
-			UE_LOG(LogVideoEncoder, Warning, TEXT("SetTexture | Cuda device pointer is null"));
+			CUDA.EncoderTexture = InTexture;
+			OnReleaseCUDATexture = InOnReleaseTexture;
+			if (!CUDA.EncoderTexture)
+			{
+				UE_LOG(LogVideoEncoder, Warning, TEXT("SetTexture | Cuda device pointer is null"));
+			}
 		}
 	}
-}
 
 
 #if PLATFORM_DESKTOP && !PLATFORM_APPLE
-void FVideoEncoderInputFrame::SetTexture(VkImage InTexture, FReleaseVulkanTextureCallback InOnReleaseTexture)
-{
-	if (Format == EVideoFrameFormat::VULKAN_R8G8B8A8_UNORM)
+	void FVideoEncoderInputFrame::SetTexture(VkImage InTexture, FReleaseVulkanTextureCallback InOnReleaseTexture)
 	{
-		Vulkan.EncoderTexture = InTexture;		
-		OnReleaseVulkanTexture = InOnReleaseTexture;
-		if (!Vulkan.EncoderTexture)
+		if (Format == EVideoFrameFormat::VULKAN_R8G8B8A8_UNORM)
 		{
-			UE_LOG(LogVideoEncoder, Warning, TEXT("SetTexture | Vulkan VkImage is null"));
+			Vulkan.EncoderTexture = InTexture;
+			OnReleaseVulkanTexture = InOnReleaseTexture;
+			if (!Vulkan.EncoderTexture)
+			{
+				UE_LOG(LogVideoEncoder, Warning, TEXT("SetTexture | Vulkan VkImage is null"));
+			}
 		}
 	}
-}
 
-void FVideoEncoderInputFrame::SetTexture(VkImage InTexture, VkDeviceMemory InTextureDeviceMemory, uint64 InTextureMemorySize, FReleaseVulkanTextureCallback InOnReleaseTexture)
-{
-	if (Format == EVideoFrameFormat::VULKAN_R8G8B8A8_UNORM)
+	void FVideoEncoderInputFrame::SetTexture(VkImage InTexture, VkDeviceMemory InTextureDeviceMemory, uint64 InTextureMemorySize, FReleaseVulkanTextureCallback InOnReleaseTexture)
 	{
-		Vulkan.EncoderTexture = InTexture;
-		Vulkan.EncoderDeviceMemory = InTextureDeviceMemory;
-		Vulkan.EncoderMemorySize = InTextureMemorySize;
-		OnReleaseVulkanTexture = InOnReleaseTexture;
-
-		if (!Vulkan.EncoderTexture || !InTextureDeviceMemory)
+		if (Format == EVideoFrameFormat::VULKAN_R8G8B8A8_UNORM)
 		{
-			UE_LOG(LogVideoEncoder, Warning, TEXT("SetTexture | Vulkan VkImage or VkTextureDeviceMemory is null"));
+			Vulkan.EncoderTexture = InTexture;
+			Vulkan.EncoderDeviceMemory = InTextureDeviceMemory;
+			Vulkan.EncoderMemorySize = InTextureMemorySize;
+			OnReleaseVulkanTexture = InOnReleaseTexture;
+
+			if (!Vulkan.EncoderTexture || !InTextureDeviceMemory)
+			{
+				UE_LOG(LogVideoEncoder, Warning, TEXT("SetTexture | Vulkan VkImage or VkTextureDeviceMemory is null"));
+			}
 		}
 	}
-}
 #endif
 
-// *** FVideoEncoderInputFrameImpl ****************************************************************
-
-FVideoEncoderInputFrameImpl::FVideoEncoderInputFrameImpl(FVideoEncoderInputImpl* InInput)
-	: Input(InInput)
-{
-}
-
-FVideoEncoderInputFrameImpl::FVideoEncoderInputFrameImpl(const FVideoEncoderInputFrameImpl& InCloneFrom, FCloneDestroyedCallback InCloneDestroyedCallback)
-	: FVideoEncoderInputFrame(InCloneFrom)
-	, Input(InCloneFrom.Input)
-	, ClonedReference(InCloneFrom.Obtain())
-	, OnCloneDestroyed(MoveTemp(InCloneDestroyedCallback))
-{
-}
+	// *** FVideoEncoderInputFrameImpl ****************************************************************
 
-FVideoEncoderInputFrameImpl::~FVideoEncoderInputFrameImpl()
-{
-	if (ClonedReference)
+	FVideoEncoderInputFrameImpl::FVideoEncoderInputFrameImpl(FVideoEncoderInputImpl* InInput)
+		: Input(InInput)
 	{
-		ClonedReference->Release();
 	}
-	else
+
+	FVideoEncoderInputFrameImpl::FVideoEncoderInputFrameImpl(const FVideoEncoderInputFrameImpl& InCloneFrom, FCloneDestroyedCallback InCloneDestroyedCallback)
+		: FVideoEncoderInputFrame(InCloneFrom)
+		, Input(InCloneFrom.Input)
+		, ClonedReference(InCloneFrom.Obtain())
+		, OnCloneDestroyed(MoveTemp(InCloneDestroyedCallback))
 	{
 	}
-}
 
-void FVideoEncoderInputFrameImpl::Release() const
-{
-	if (NumReferences.Decrement() == 0)
+	FVideoEncoderInputFrameImpl::~FVideoEncoderInputFrameImpl()
 	{
 		if (ClonedReference)
 		{
-			OnCloneDestroyed(this);
-			delete this;
+			ClonedReference->Release();
 		}
 		else
+		{
+		}
+	}
+
+	void FVideoEncoderInputFrameImpl::Release() const
+	{
+		// User managed frames get released without checking reference count, as the reference count doesn't matter for them
+		if (Input->IsUserManagedFrame(this))
 		{
 			Input->ReleaseInputFrame(const_cast<FVideoEncoderInputFrameImpl*>(this));
 		}
+		else if (NumReferences.Decrement() == 0)
+		{
+			if (ClonedReference)
+			{
+				OnCloneDestroyed(this);
+				delete this;
+			}
+			else
+			{
+				Input->ReleaseInputFrame(const_cast<FVideoEncoderInputFrameImpl*>(this));
+			}
+		}
 	}
-}
 
-// Clone frame - this will create a copy that references the original until destroyed
-const FVideoEncoderInputFrame* FVideoEncoderInputFrameImpl::Clone(FCloneDestroyedCallback InCloneDestroyedCallback) const
-{
-	FVideoEncoderInputFrameImpl*	ClonedFrame = new FVideoEncoderInputFrameImpl(*this, MoveTemp(InCloneDestroyedCallback));
-	return ClonedFrame;
-}
+	// Clone frame - this will create a copy that references the original until destroyed
+	const FVideoEncoderInputFrame* FVideoEncoderInputFrameImpl::Clone(FCloneDestroyedCallback InCloneDestroyedCallback) const
+	{
+		FVideoEncoderInputFrameImpl* ClonedFrame = new FVideoEncoderInputFrameImpl(*this, MoveTemp(InCloneDestroyedCallback));
+		return ClonedFrame;
+	}
 
 
 
diff --git a/Engine/Source/Runtime/AVEncoder/Private/VideoEncoderInputImpl.h b/Engine/Source/Runtime/AVEncoder/Private/VideoEncoderInputImpl.h
--- a/Engine/Source/Runtime/AVEncoder/Private/VideoEncoderInputImpl.h
+++ b/Engine/Source/Runtime/AVEncoder/Private/VideoEncoderInputImpl.h
@@ -76,6 +76,9 @@ public:
 	// destroy/release any frames that are not currently in use
 	void Flush() override;
 
+	// indicates whether a given frame is a user managed frame or not
+	bool IsUserManagedFrame(const FVideoEncoderInputFrame* InBuffer) const;
+
 	// --- input properties
 
 #if PLATFORM_WINDOWS
