removing some rogue references to PLATFORM_XBOXONE and WITH_LEGACY_XDK.

#rnx

diff --git a/Engine/Source/Runtime/AVEncoder/Private/IMFSampleWrapper.cpp b/Engine/Source/Runtime/AVEncoder/Private/IMFSampleWrapper.cpp
--- a/Engine/Source/Runtime/AVEncoder/Private/IMFSampleWrapper.cpp
+++ b/Engine/Source/Runtime/AVEncoder/Private/IMFSampleWrapper.cpp
@@ -3,7 +3,7 @@
 #include "IMFSampleWrapper.h"
 #include "AVEncoder.h"
 
-#if PLATFORM_WINDOWS || (PLATFORM_XBOXONE && WITH_LEGACY_XDK)
+#if PLATFORM_WINDOWS
 
 #include "MicrosoftCommon.h"
 
@@ -121,5 +121,5 @@ FIMFSampleWrapper FIMFSampleWrapper::Clone() const
 
 } // namespace AVEncoder
 
-#endif // PLATFORM_WINDOWS || (PLATFORM_XBOXONE && WITH_LEGACY_XDK)
+#endif // PLATFORM_WINDOWS
 
diff --git a/Engine/Source/Runtime/AVEncoder/Private/MicrosoftCommon.cpp b/Engine/Source/Runtime/AVEncoder/Private/MicrosoftCommon.cpp
--- a/Engine/Source/Runtime/AVEncoder/Private/MicrosoftCommon.cpp
+++ b/Engine/Source/Runtime/AVEncoder/Private/MicrosoftCommon.cpp
@@ -2,7 +2,7 @@
 
 #include "MicrosoftCommon.h"
 
-#if PLATFORM_WINDOWS || (PLATFORM_XBOXONE && WITH_LEGACY_XDK)
+#if PLATFORM_WINDOWS
 
 #if WMFMEDIA_SUPPORTED_PLATFORM
 	#pragma comment(lib, "mfplat")
@@ -41,5 +41,5 @@ namespace AVEncoder
 
 } // namespace AVEncoder
 
-#endif // PLATFORM_WINDOWS || (PLATFORM_XBOXONE && WITH_LEGACY_XDK)
+#endif // PLATFORM_WINDOWS
 
diff --git a/Engine/Source/Runtime/AVEncoder/Private/MicrosoftCommon.h b/Engine/Source/Runtime/AVEncoder/Private/MicrosoftCommon.h
--- a/Engine/Source/Runtime/AVEncoder/Private/MicrosoftCommon.h
+++ b/Engine/Source/Runtime/AVEncoder/Private/MicrosoftCommon.h
@@ -5,7 +5,7 @@
 #include "Containers/UnrealString.h"
 #include "VideoCommon.h"
 
-#if PLATFORM_WINDOWS || (PLATFORM_XBOXONE && WITH_LEGACY_XDK) || PLATFORM_HOLOLENS
+#if PLATFORM_WINDOWS || PLATFORM_HOLOLENS
 
 namespace AVEncoder
 {
@@ -163,6 +163,6 @@ const GUID CLSID_VideoProcessorMFT = { 0x88753b26, 0x5b24, 0x49bd, { 0xb2, 0xe7,
 } // namespace AVEncoder
 
 
-#endif // PLATFORM_WINDOWS || (PLATFORM_XBOXONE && WITH_LEGACY_XDK)
+#endif // PLATFORM_WINDOWS || PLATFORM_HOLOLENS
 
 
diff --git a/Engine/Source/Runtime/AVEncoder/Public/IMFSampleWrapper.h b/Engine/Source/Runtime/AVEncoder/Public/IMFSampleWrapper.h
--- a/Engine/Source/Runtime/AVEncoder/Public/IMFSampleWrapper.h
+++ b/Engine/Source/Runtime/AVEncoder/Public/IMFSampleWrapper.h
@@ -6,7 +6,7 @@
 #include "Templates/RefCounting.h"
 #include "VideoCommon.h"
 
-#if PLATFORM_WINDOWS || (PLATFORM_XBOXONE && WITH_LEGACY_XDK)
+#if PLATFORM_WINDOWS
 
 namespace AVEncoder
 {
@@ -100,5 +100,5 @@ private:
 
 }
 
-#endif // PLATFORM_WINDOWS || (PLATFORM_XBOXONE && WITH_LEGACY_XDK)
+#endif // PLATFORM_WINDOWS
 
diff --git a/Engine/Source/Runtime/AVEncoder/Public/VideoCommon.h b/Engine/Source/Runtime/AVEncoder/Public/VideoCommon.h
--- a/Engine/Source/Runtime/AVEncoder/Public/VideoCommon.h
+++ b/Engine/Source/Runtime/AVEncoder/Public/VideoCommon.h
@@ -34,33 +34,6 @@ THIRD_PARTY_INCLUDES_END
 
 #endif // PLATFORM_WINDOWS
 
-//
-// XboxOne only includes
-//
-#if (PLATFORM_XBOXONE && WITH_LEGACY_XDK)
-
-#pragma warning(push)
-#pragma warning(disable: 4005)
-
-THIRD_PARTY_INCLUDES_START
-#include "XboxCommonAllowPlatformTypes.h"
-#include "XboxCommonPreApi.h"
-#include <d3d11_x.h>
-#include <d3d12_x.h>
-#include <d3dx12_x.h>
-#include <mfobjects.h>
-#include <mftransform.h>
-#include <mfapi.h>
-#include <mferror.h>
-#include <mfidl.h>
-#include <codecapi.h>
-#include <mfreadwrite.h>
-#include "XboxCommonPostApi.h"
-#include "XboxCommonHidePlatformTypes.h"
-THIRD_PARTY_INCLUDES_END
-
-#endif  // (PLATFORM_XBOXONE && WITH_LEGACY_XDK)
-
 #ifndef WMFMEDIA_SUPPORTED_PLATFORM
 	#define WMFMEDIA_SUPPORTED_PLATFORM (PLATFORM_WINDOWS && (WINVER >= 0x0600 /*Vista*/) && !UE_SERVER)
 #endif
