removing several PLATFORM_HOLOLENS references, mostly around WindowsHWrapper etc.

#rnx

diff --git a/Engine/Plugins/Enterprise/DataprepEditor/Source/DataprepEditor/Private/Widgets/DataprepGraph/SDataprepGraphTrackNode.cpp b/Engine/Plugins/Enterprise/DataprepEditor/Source/DataprepEditor/Private/Widgets/DataprepGraph/SDataprepGraphTrackNode.cpp
--- a/Engine/Plugins/Enterprise/DataprepEditor/Source/DataprepEditor/Private/Widgets/DataprepGraph/SDataprepGraphTrackNode.cpp
+++ b/Engine/Plugins/Enterprise/DataprepEditor/Source/DataprepEditor/Private/Widgets/DataprepGraph/SDataprepGraphTrackNode.cpp
@@ -28,8 +28,8 @@
 #include "Widgets/SBoxPanel.h"
 #include "Widgets/SOverlay.h"
 
-#if PLATFORM_WINDOWS || PLATFORM_HOLOLENS
-#include "Windows/WindowsHWrapper.h"
+#if PLATFORM_MICROSOFT
+#include "Microsoft/WindowsHWrapper.h"
 #endif
 
 #define LOCTEXT_NAMESPACE "DataprepGraphEditor"
diff --git a/Engine/Plugins/Media/HLMedia/Source/HLMedia/Private/HLMediaPrivate.h b/Engine/Plugins/Media/HLMedia/Source/HLMedia/Private/HLMediaPrivate.h
--- a/Engine/Plugins/Media/HLMedia/Source/HLMedia/Private/HLMediaPrivate.h
+++ b/Engine/Plugins/Media/HLMedia/Source/HLMedia/Private/HLMediaPrivate.h
@@ -8,9 +8,9 @@
 #include "Containers/UnrealString.h"
 #include "Logging/LogMacros.h"
 
-#if (defined(PLATFORM_WINDOWS) || defined(PLATFORM_HOLOLENS))
-#include "Windows/WindowsHWrapper.h"
-#include "Windows/AllowWindowsPlatformTypes.h"
+#if PLATFORM_MICROSOFT
+#include "Microsoft/WindowsHWrapper.h"
+#include "Microsoft/AllowMicrosoftPlatformTypes.h"
 
 #if PLATFORM_DESKTOP
     #if PLATFORM_64BITS
@@ -51,8 +51,8 @@ THIRD_PARTY_INCLUDES_END
 
 #include "Microsoft/COMPointer.h"
 
-#if PLATFORM_WINDOWS || PLATFORM_HOLOLENS
-#include "Windows/HideWindowsPlatformTypes.h"
+#if PLATFORM_MICROSOFT
+#include "Microsoft/HideMicrosoftPlatformTypes.h"
 #endif
 
 // General Log
diff --git a/Engine/Source/Developer/DerivedDataCache/Private/HttpCacheStore.cpp b/Engine/Source/Developer/DerivedDataCache/Private/HttpCacheStore.cpp
--- a/Engine/Source/Developer/DerivedDataCache/Private/HttpCacheStore.cpp
+++ b/Engine/Source/Developer/DerivedDataCache/Private/HttpCacheStore.cpp
@@ -50,12 +50,12 @@
 
 #include "Http/HttpClient.h"
 
-#if PLATFORM_WINDOWS || PLATFORM_HOLOLENS
-#include "Windows/WindowsHWrapper.h"
-#include "Windows/AllowWindowsPlatformTypes.h"
+#if PLATFORM_MICROSOFT
+#include "Microsoft/WindowsHWrapper.h"
+#include "Microsoft/AllowMicrosoftPlatformTypes.h"
 #include <winsock2.h>
 #include <ws2tcpip.h>
-#include "Windows/HideWindowsPlatformTypes.h"
+#include "Microsoft/HideMicrosoftPlatformTypes.h"
 #else
 #include <sys/types.h>
 #include <sys/socket.h>
diff --git a/Engine/Source/Developer/DerivedDataCache/Private/S3CacheStore.cpp b/Engine/Source/Developer/DerivedDataCache/Private/S3CacheStore.cpp
--- a/Engine/Source/Developer/DerivedDataCache/Private/S3CacheStore.cpp
+++ b/Engine/Source/Developer/DerivedDataCache/Private/S3CacheStore.cpp
@@ -4,12 +4,12 @@
 
 #if WITH_S3_DDC_BACKEND
 
-#if PLATFORM_WINDOWS || PLATFORM_HOLOLENS
-	#include "Windows/WindowsHWrapper.h"
-	#include "Windows/AllowWindowsPlatformTypes.h"
+#if PLATFORM_MICROSOFT
+	#include "Microsoft/WindowsHWrapper.h"
+	#include "Microsoft/AllowMicrosoftPlatformTypes.h"
 #endif
-#if PLATFORM_WINDOWS || PLATFORM_HOLOLENS
-	#include "Windows/HideWindowsPlatformTypes.h"
+#if PLATFORM_MICROSOFT
+	#include "Microsoft/HideMicrosoftPlatformTypes.h"
 #endif
 
 #include "Algo/AllOf.h"
diff --git a/Engine/Source/Developer/Virtualization/Private/VirtualizationHttpBackend.cpp b/Engine/Source/Developer/Virtualization/Private/VirtualizationHttpBackend.cpp
--- a/Engine/Source/Developer/Virtualization/Private/VirtualizationHttpBackend.cpp
+++ b/Engine/Source/Developer/Virtualization/Private/VirtualizationHttpBackend.cpp
@@ -28,15 +28,15 @@
 #include "Serialization/MemoryWriter.h"
 
 // The following inlcudes are all for the Utility namespace
-#if PLATFORM_WINDOWS || PLATFORM_HOLOLENS
-	#include "Windows/WindowsHWrapper.h"
-	#include "Windows/AllowWindowsPlatformTypes.h"
+#if PLATFORM_MICROSOFT
+	#include "Microsoft/WindowsHWrapper.h"
+	#include "Microsoft/AllowMicrosoftPlatformTypes.h"
 #endif
 
 #include "curl/curl.h"
 
-#if PLATFORM_WINDOWS || PLATFORM_HOLOLENS
-	#include "Windows/HideWindowsPlatformTypes.h"
+#if PLATFORM_MICROSOFT
+	#include "Microsoft/HideMicrosoftPlatformTypes.h"
 #endif
 
 #include "Containers/StaticArray.h"
diff --git a/Engine/Source/Developer/Zen/Internal/ZenServerHttp.cpp b/Engine/Source/Developer/Zen/Internal/ZenServerHttp.cpp
--- a/Engine/Source/Developer/Zen/Internal/ZenServerHttp.cpp
+++ b/Engine/Source/Developer/Zen/Internal/ZenServerHttp.cpp
@@ -7,9 +7,9 @@
 #include "ZenBackendUtils.h"
 #include "ZenSerialization.h"
 
-#if PLATFORM_WINDOWS || PLATFORM_HOLOLENS
-#	include "Windows/WindowsHWrapper.h"
-#	include "Windows/AllowWindowsPlatformTypes.h"
+#if PLATFORM_MICROSOFT
+#	include "Microsoft/WindowsHWrapper.h"
+#	include "Microsoft/AllowMicrosoftPlatformTypes.h"
 #endif
 
 #if PLATFORM_WINDOWS
@@ -20,8 +20,8 @@
 #include "curl/curl.h"
 #endif
 
-#if PLATFORM_WINDOWS || PLATFORM_HOLOLENS
-#	include "Windows/HideWindowsPlatformTypes.h"
+#if PLATFORM_MICROSOFT
+#	include "Microsoft/HideMicrosoftPlatformTypes.h"
 #endif
 
 #include "Logging/LogMacros.h"
diff --git a/Engine/Source/Runtime/AVEncoder/Private/MicrosoftCommon.h b/Engine/Source/Runtime/AVEncoder/Private/MicrosoftCommon.h
--- a/Engine/Source/Runtime/AVEncoder/Private/MicrosoftCommon.h
+++ b/Engine/Source/Runtime/AVEncoder/Private/MicrosoftCommon.h
@@ -5,7 +5,7 @@
 #include "Containers/UnrealString.h"
 #include "VideoCommon.h"
 
-#if PLATFORM_WINDOWS || PLATFORM_HOLOLENS
+#if PLATFORM_MICROSOFT
 
 namespace AVEncoder
 {
@@ -30,11 +30,7 @@ inline const FString GetComErrorDescription(HRESULT Res)
 	}
 }
 
-#if PLATFORM_WINDOWS || PLATFORM_HOLOLENS
-	#include "Windows/AllowWindowsPlatformTypes.h"
-#elif PLATFORM_XBOXONE
-	#include "XboxCommonAllowPlatformTypes.h"
-#endif
+#include "Microsoft/AllowMicrosoftPlatformTypes.h"
 
 // macro to deal with COM calls inside a function that returns `false` on error
 #define CHECK_HR(COM_call)\
@@ -80,11 +76,7 @@ inline const FString GetComErrorDescription(HRESULT Res)
 		}\
 	}
 
-#if PLATFORM_WINDOWS || PLATFORM_HOLOLENS
-	#include "Windows/HideWindowsPlatformTypes.h"
-#elif PLATFORM_XBOXONE
-	#include "XboxCommonHidePlatformTypes.h"
-#endif
+#include "Microsoft/HideMicrosoftPlatformTypes.h"
 
 // following commented include causes name clash between UE4 and Windows `IMediaEventSink`,
 // we just need a couple of GUIDs from there so the solution is to duplicate them below
@@ -163,6 +155,6 @@ const GUID CLSID_VideoProcessorMFT = { 0x88753b26, 0x5b24, 0x49bd, { 0xb2, 0xe7,
 } // namespace AVEncoder
 
 
-#endif // PLATFORM_WINDOWS || PLATFORM_HOLOLENS
+#endif // PLATFORM_MICROSOFT
 
 
diff --git a/Engine/Source/Runtime/AVEncoder/Public/VideoCommon.h b/Engine/Source/Runtime/AVEncoder/Public/VideoCommon.h
--- a/Engine/Source/Runtime/AVEncoder/Public/VideoCommon.h
+++ b/Engine/Source/Runtime/AVEncoder/Public/VideoCommon.h
@@ -14,8 +14,8 @@
 #pragma warning(disable: 4005)
 
 THIRD_PARTY_INCLUDES_START
-#include "Windows/AllowWindowsPlatformTypes.h"
-#include "Windows/PreWindowsApi.h"
+#include "Microsoft/AllowMicrosoftPlatformTypes.h"
+#include "Microsoft/PreWindowsApi.h"
 #include <d3d11.h>
 #include <mfobjects.h>
 #include <mftransform.h>
@@ -28,8 +28,8 @@ THIRD_PARTY_INCLUDES_START
 #include <d3d11_1.h>
 #include <d3d12.h>
 #include <dxgi1_4.h>
-#include "Windows/PostWindowsApi.h"
-#include "Windows/HideWindowsPlatformTypes.h"
+#include "Microsoft/PostWindowsApi.h"
+#include "Microsoft/HideMicrosoftPlatformTypes.h"
 THIRD_PARTY_INCLUDES_END
 
 #endif // PLATFORM_WINDOWS
diff --git a/Engine/Source/Runtime/D3D12RHI/Private/D3D12RootSignature.cpp b/Engine/Source/Runtime/D3D12RHI/Private/D3D12RootSignature.cpp
--- a/Engine/Source/Runtime/D3D12RHI/Private/D3D12RootSignature.cpp
+++ b/Engine/Source/Runtime/D3D12RHI/Private/D3D12RootSignature.cpp
@@ -17,9 +17,9 @@
 #endif
 
 #if !PLATFORM_CPU_ARM_FAMILY && (PLATFORM_WINDOWS || PLATFORM_HOLOLENS)
-#include "Windows/AllowWindowsPlatformTypes.h"
+#include "Microsoft/AllowMicrosoftPlatformTypes.h"
 #include "amd_ags.h"
-#include "Windows/HideWindowsPlatformTypes.h"
+#include "Microsoft/HideMicrosoftPlatformTypes.h"
 #endif
 
 // Allows to automatically bind UEDiagnosticBuffer UAV, available to all shaders.
diff --git a/Engine/Source/Runtime/Online/ICMP/Private/IcmpPrivate.h b/Engine/Source/Runtime/Online/ICMP/Private/IcmpPrivate.h
--- a/Engine/Source/Runtime/Online/ICMP/Private/IcmpPrivate.h
+++ b/Engine/Source/Runtime/Online/ICMP/Private/IcmpPrivate.h
@@ -14,12 +14,6 @@
 		#define PLATFORM_USES_POSIX_ICMP 0
 		#define PING_ALLOWS_CUSTOM_THREAD_SIZE 1
 
-	#elif PLATFORM_HOLOLENS
-
-		#define PLATFORM_SUPPORTS_ICMP 0
-		#define PLATFORM_USES_POSIX_ICMP 0
-		#define PING_ALLOWS_CUSTOM_THREAD_SIZE 0
-
 	#elif PLATFORM_MAC
 
 		#define PLATFORM_SUPPORTS_ICMP 1
diff --git a/Engine/Source/Runtime/Slate/Private/Framework/Application/SlateApplication.cpp b/Engine/Source/Runtime/Slate/Private/Framework/Application/SlateApplication.cpp
--- a/Engine/Source/Runtime/Slate/Private/Framework/Application/SlateApplication.cpp
+++ b/Engine/Source/Runtime/Slate/Private/Framework/Application/SlateApplication.cpp
@@ -51,8 +51,8 @@
 	#define SLATE_HAS_WIDGET_REFLECTOR !(UE_BUILD_TEST || UE_BUILD_SHIPPING) && PLATFORM_DESKTOP
 #endif
 
-#if PLATFORM_WINDOWS || PLATFORM_HOLOLENS
-#include "Windows/WindowsHWrapper.h"
+#if PLATFORM_MICROSOFT
+#include "Microsoft/WindowsHWrapper.h"
 #endif
 #include "Debugging/SlateDebugging.h"
 #include "Styling/StyleColors.h"
