Fix DEBUG_SET_D3D11_OBJECT_NAME for Edit & Continue.

 - Cannot use __LINE__ as a compile-time constant because it cannot be used with Edit & Continue (/ZI)
 - see C2975 in MSDN

#rnx

diff --git a/Engine/Source/Runtime/AVEncoder/Public/AVEncoderDebug.h b/Engine/Source/Runtime/AVEncoder/Public/AVEncoderDebug.h
--- a/Engine/Source/Runtime/AVEncoder/Public/AVEncoderDebug.h
+++ b/Engine/Source/Runtime/AVEncoder/Public/AVEncoderDebug.h
@@ -5,7 +5,7 @@
 #define DEBUG_SET_D3D11_OBJECT_NAME(D3D11Object, Name) \
 {\
 	static GUID _D3DDebugObjectNameW = {0x4cca5fd8,0x921f,0x42c8,{0x85,0x66,0x70,0xca,0xf2,0xa9,0xb7,0x41}}; \
-	FString		SetName = FString::Format(TEXT("{0} ({1}:{2})"), {FString(Name), FPaths::GetCleanFilename(FString(__FILE__)), __LINE__}); \
+	FString		SetName = FString::Printf(TEXT("%s (%s:%d)"), *FString(Name), *FPaths::GetCleanFilename(FString(__FILE__)), __LINE__); \
 	D3D11Object->SetPrivateData(_D3DDebugObjectNameW, SetName.Len() * sizeof(TCHAR), *SetName); \
 }
 
