Fix malloc(sizeof) and cast mismatch picked up by Clang static analysis in AVEncoder

diff --git a/Engine/Source/Runtime/AVEncoder/Private/Decoders/vdecmpeg4/M4Decoder.cpp b/Engine/Source/Runtime/AVEncoder/Private/Decoders/vdecmpeg4/M4Decoder.cpp
--- a/Engine/Source/Runtime/AVEncoder/Private/Decoders/vdecmpeg4/M4Decoder.cpp
+++ b/Engine/Source/Runtime/AVEncoder/Private/Decoders/vdecmpeg4/M4Decoder.cpp
@@ -254,7 +254,7 @@ VIDError M4Decoder::initBuffers(int16 width, int16 height)
 	freeBuffers();
 
 	// Create at least 2 (non-B frames) or 3 (B-frames) buffers here
-	mpVidImage = (M4Image**)mMemSys.malloc(sizeof(M4Image**) * mNumVidImages);
+	mpVidImage = (M4Image**)mMemSys.malloc(sizeof(M4Image*) * mNumVidImages);
 	for(uint32 i=0; i < mNumVidImages; ++i)
 	{
 		mpVidImage[i] = M4Image::create(this, mWidth, mHeight);
