Fixing a crash in GameplayMediaEncoder that would trigger when starting a highlight recording.

diff --git a/Engine/Source/Runtime/AVEncoder/Private/VideoEncoderInput.cpp b/Engine/Source/Runtime/AVEncoder/Private/VideoEncoderInput.cpp
--- a/Engine/Source/Runtime/AVEncoder/Private/VideoEncoderInput.cpp
+++ b/Engine/Source/Runtime/AVEncoder/Private/VideoEncoderInput.cpp
@@ -161,13 +161,6 @@ namespace AVEncoder
 			}
 
 			check(ActiveFrames.Num() == 0);
-
-			while (!AvailableFrames.IsEmpty())
-			{
-				FVideoEncoderInputFrameImpl* Frame = nullptr;
-				AvailableFrames.Dequeue(Frame);
-				delete Frame;
-			}
 		}
 #if PLATFORM_WINDOWS
 		//	DEBUG_D3D11_REPORT_LIVE_DEVICE_OBJECT(FrameInfoD3D.EncoderDeviceD3D11);
@@ -482,9 +475,9 @@ namespace AVEncoder
 
 	// --- encoder input frames -----------------------------------------------------------------------
 
-	FVideoEncoderInputFrame* FVideoEncoderInputImpl::ObtainInputFrame()
+	TSharedPtr<FVideoEncoderInputFrame> FVideoEncoderInputImpl::ObtainInputFrame()
 	{
-		FVideoEncoderInputFrameImpl* Frame = nullptr;
+		TSharedPtr<FVideoEncoderInputFrameImpl> Frame;
 		FScopeLock						Guard(&ProtectFrames);
 
 		if (!AvailableFrames.IsEmpty())
@@ -495,7 +488,7 @@ namespace AVEncoder
 		}
 		else
 		{
-			Frame = CreateFrame();
+			Frame = MakeShareable(CreateFrame());
 			UE_LOG(LogVideoEncoder, Verbose, TEXT("Created new frame total frames: %d"), NumBuffers);
 		}
 
@@ -508,7 +501,8 @@ namespace AVEncoder
 			++NextFrameID; // skip 0 id
 		}
 
-		return const_cast<FVideoEncoderInputFrame*>(Frame->Obtain());
+		Frame->Obtain();
+		return Frame;
 	}
 
 	FVideoEncoderInputFrameImpl* FVideoEncoderInputImpl::CreateFrame()
@@ -557,7 +551,15 @@ namespace AVEncoder
 			}
 		}
 
-		int32 NumRemoved = ActiveFrames.Remove(InFrameImpl);
+		TSharedPtr<FVideoEncoderInputFrameImpl>* InFramePtrPtr = ActiveFrames.FindByPredicate([InFrameImpl](TSharedPtr<FVideoEncoderInputFrameImpl> ActiveFrame) { return ActiveFrame.Get() == InFrameImpl; });
+		if (!InFramePtrPtr)
+		{
+			// releasing a non active frame. might be after we flushed or something. ignore it.
+			return;
+		}
+
+		TSharedPtr<FVideoEncoderInputFrameImpl> InFramePtr = *InFramePtrPtr;
+		int32 NumRemoved = ActiveFrames.Remove(InFramePtr);
 		check(NumRemoved == 1);
 		if (NumRemoved > 0)
 		{
@@ -565,7 +567,6 @@ namespace AVEncoder
 			if (InFrame->GetFormat() != FrameFormat)
 			{
 				ProtectFrames.Unlock();
-				delete InFrameImpl;
 				NumBuffers--;
 				UE_LOG(LogVideoEncoder, Verbose, TEXT("Deleted buffer (format mismatch) total remaining: %d"), NumBuffers);
 				return;
@@ -574,29 +575,20 @@ namespace AVEncoder
 			if (!AvailableFrames.IsEmpty() && NumBuffers > MaxNumBuffers)
 			{
 				ProtectFrames.Unlock();
-				delete InFrameImpl;
 				NumBuffers--;
 				UE_LOG(LogVideoEncoder, Verbose, TEXT("Deleted buffer (too many) total frames: %d"), NumBuffers);
 				return;
 			}
 
-			AvailableFrames.Enqueue(InFrameImpl);
+			AvailableFrames.Enqueue(InFramePtr);
 		}
 	}
 
 	void FVideoEncoderInputImpl::Flush()
 	{
-		ProtectFrames.Lock();
-		while (!AvailableFrames.IsEmpty())
-		{
-			FVideoEncoderInputFrameImpl* Frame = nullptr;
-			AvailableFrames.Dequeue(Frame);
-			ProtectFrames.Unlock();
-			delete Frame;
-			NumBuffers--;
-			ProtectFrames.Lock();
-		}
-		ProtectFrames.Unlock();
+		FScopeLock ScopeLock(&ProtectFrames);
+		AvailableFrames.Empty();
+		NumBuffers = 0;
 	}
 
 	void FVideoEncoderInputImpl::SetupFrameYUV420P(FVideoEncoderInputFrameImpl* Frame)
diff --git a/Engine/Source/Runtime/AVEncoder/Private/VideoEncoderInputImpl.h b/Engine/Source/Runtime/AVEncoder/Private/VideoEncoderInputImpl.h
--- a/Engine/Source/Runtime/AVEncoder/Private/VideoEncoderInputImpl.h
+++ b/Engine/Source/Runtime/AVEncoder/Private/VideoEncoderInputImpl.h
@@ -68,7 +68,7 @@ public:
 	// --- encoder input frames - managed by this object
 
 	// obtain a video frame that can be used as a buffer for input to a video encoder
-	FVideoEncoderInputFrame* ObtainInputFrame() override;
+	TSharedPtr<FVideoEncoderInputFrame> ObtainInputFrame() override;
 
 	// release (free) an input frame and make it available for future use
 	void ReleaseInputFrame(FVideoEncoderInputFrame* InFrame) override;
@@ -136,8 +136,8 @@ private:
 #endif
 
 	mutable FCriticalSection				ProtectFrames;
-	TQueue<FVideoEncoderInputFrameImpl*>	AvailableFrames;
-	TArray<FVideoEncoderInputFrameImpl*>	ActiveFrames;
+	TQueue<TSharedPtr<FVideoEncoderInputFrameImpl>>	AvailableFrames;
+	TArray<TSharedPtr<FVideoEncoderInputFrameImpl>>	ActiveFrames;
 	using UserManagedFrame = TPair<FVideoEncoderInputFrameImpl*, OnFrameReleasedCallback>;
 	TArray<UserManagedFrame>				UserManagedFrames;
 };
diff --git a/Engine/Source/Runtime/AVEncoder/Public/VideoEncoderInput.h b/Engine/Source/Runtime/AVEncoder/Public/VideoEncoderInput.h
--- a/Engine/Source/Runtime/AVEncoder/Public/VideoEncoderInput.h
+++ b/Engine/Source/Runtime/AVEncoder/Public/VideoEncoderInput.h
@@ -87,7 +87,7 @@ namespace AVEncoder
 		// --- encoder input frames - managed by this object
 
 		// obtain a video frame that can be used as a buffer for input to a video encoder
-		virtual FVideoEncoderInputFrame* ObtainInputFrame() = 0;
+		virtual TSharedPtr<FVideoEncoderInputFrame> ObtainInputFrame() = 0;
 
 		// release (free) an input frame and make it available for future use
 		virtual void ReleaseInputFrame(FVideoEncoderInputFrame* InFrame) = 0;
diff --git a/Engine/Source/Runtime/GameplayMediaEncoder/Private/GameplayMediaEncoder.cpp b/Engine/Source/Runtime/GameplayMediaEncoder/Private/GameplayMediaEncoder.cpp
--- a/Engine/Source/Runtime/GameplayMediaEncoder/Private/GameplayMediaEncoder.cpp
+++ b/Engine/Source/Runtime/GameplayMediaEncoder/Private/GameplayMediaEncoder.cpp
@@ -576,8 +576,9 @@ void FGameplayMediaEncoder::ProcessVideoFrame(const FTexture2DRHIRef& FrameBuffe
 
 TSharedPtr<AVEncoder::FVideoEncoderInputFrame> FGameplayMediaEncoder::ObtainInputFrame()
 {
-	AVEncoder::FVideoEncoderInputFrame* Frame = VideoEncoderInput->ObtainInputFrame();
-	TSharedPtr<AVEncoder::FVideoEncoderInputFrame> InputFrame = MakeShareable(Frame);
+	TSharedPtr<AVEncoder::FVideoEncoderInputFrame> InputFrame = VideoEncoderInput->ObtainInputFrame();
+	InputFrame->SetWidth(VideoConfig.Width);
+	InputFrame->SetHeight(VideoConfig.Height);
 
 	if(!BackBuffers.Contains(InputFrame))
 	{
