Added input parameter checking.

#rnx

diff --git a/Engine/Source/Runtime/AVEncoder/Private/Decoders/Windows/VideoDecoderH264_Windows.cpp b/Engine/Source/Runtime/AVEncoder/Private/Decoders/Windows/VideoDecoderH264_Windows.cpp
--- a/Engine/Source/Runtime/AVEncoder/Private/Decoders/Windows/VideoDecoderH264_Windows.cpp
+++ b/Engine/Source/Runtime/AVEncoder/Private/Decoders/Windows/VideoDecoderH264_Windows.cpp
@@ -2252,6 +2252,11 @@ FVideoDecoder::EDecodeResult FVideoDecoderH264_WindowsImpl::Decode(const FVideoD
 	bool bOk = SetupFirstUseResources();
 	if (bOk)
 	{
+		if (!InInput->GetData() || InInput->GetDataSize() <= 0)
+		{
+			return FVideoDecoder::EDecodeResult::Failure;
+		}
+
 		// Need to wait for the decoder thread to have spun up??
 
 		TUniquePtr<FInputToDecode> Input = MakeUnique<FInputToDecode>();
