Added parameter checking

#rnx

diff --git a/Engine/Source/Runtime/AVEncoder/Private/Decoders/VideoDecoderMpeg4.cpp b/Engine/Source/Runtime/AVEncoder/Private/Decoders/VideoDecoderMpeg4.cpp
--- a/Engine/Source/Runtime/AVEncoder/Private/Decoders/VideoDecoderMpeg4.cpp
+++ b/Engine/Source/Runtime/AVEncoder/Private/Decoders/VideoDecoderMpeg4.cpp
@@ -448,6 +448,11 @@ FVideoDecoder::EDecodeResult FVideoDecoderMPEG4_Impl::Decode(const FVideoDecoder
 				ap.Width = Width;
 				ap.Height = Height;
 				ap.BytesPerPixel = 1;
+				// Check for valid values.
+				if (Width <= 0 || Height <= 0 || ap.AllocSize <= 0)
+				{
+					return FVideoDecoder::EDecodeResult::Failure;
+				}
 				ar = AllocateOutputFrameBuffer(pNew->GetBuffer(), &ap);
 				if (ar == EFrameBufferAllocReturn::CODEC_Success)
 				{
