Added an input size check

#rnx

diff --git a/Engine/Source/Runtime/AVEncoder/Private/Decoders/VideoDecoderMpeg4.cpp b/Engine/Source/Runtime/AVEncoder/Private/Decoders/VideoDecoderMpeg4.cpp
--- a/Engine/Source/Runtime/AVEncoder/Private/Decoders/VideoDecoderMpeg4.cpp
+++ b/Engine/Source/Runtime/AVEncoder/Private/Decoders/VideoDecoderMpeg4.cpp
@@ -392,6 +392,10 @@ FVideoDecoder::EDecodeResult FVideoDecoderMPEG4_Impl::Decode(const FVideoDecoder
 
 	// Setup an access unit to run through the decoder.
 	TUniquePtr<FInDecoderData> AU = MakeUnique<FInDecoderData>();
+	if (InInput->GetDataSize() <= 0)
+	{
+		return FVideoDecoder::EDecodeResult::Failure;
+	}
 	AU->Data.AddUninitialized(InInput->GetDataSize());
 	FMemory::Memcpy(AU->Data.GetData(), InInput->GetData(), InInput->GetDataSize());
 	AU->DataOffset = 0;
