QOL: Deprecate AVEncoder (for removal) and its dependencies (to be moved to plugins)

diff --git a/Engine/Plugins/Media/HardwareEncoders/Source/EncoderAMF/Private/Amf_EncoderH264.cpp b/Engine/Plugins/Media/HardwareEncoders/Source/EncoderAMF/Private/Amf_EncoderH264.cpp
--- a/Engine/Plugins/Media/HardwareEncoders/Source/EncoderAMF/Private/Amf_EncoderH264.cpp
+++ b/Engine/Plugins/Media/HardwareEncoders/Source/EncoderAMF/Private/Amf_EncoderH264.cpp
@@ -24,6 +24,7 @@
 
 namespace
 {
+	PRAGMA_DISABLE_DEPRECATION_WARNINGS
 	AMF_VIDEO_ENCODER_RATE_CONTROL_METHOD_ENUM ConvertRateControlModeAMF(AVEncoder::FVideoEncoder::RateControlMode mode)
 	{
 		switch (mode)
@@ -48,6 +49,7 @@ namespace
 			case AVEncoder::FVideoEncoder::H264Profile::AUTO: return AMF_VIDEO_ENCODER_PROFILE_BASELINE;
 		}
 	}
+	PRAGMA_ENABLE_DEPRECATION_WARNINGS
 }
 
 namespace AVEncoder
@@ -96,46 +98,62 @@ namespace AVEncoder
 		});
 	}
 
+	PRAGMA_DISABLE_DEPRECATION_WARNINGS
 	static bool GetEncoderInfo(FAmfCommon& Amf, FVideoEncoderInfo& EncoderInfo);
+	PRAGMA_ENABLE_DEPRECATION_WARNINGS
 
+	PRAGMA_DISABLE_DEPRECATION_WARNINGS
 	bool FVideoEncoderAmf_H264::GetIsAvailable(FVideoEncoderInputImpl& InInput, FVideoEncoderInfo& OutEncoderInfo)
+	PRAGMA_ENABLE_DEPRECATION_WARNINGS
 	{
 		FAmfCommon& Amf = FAmfCommon::Setup();
 		bool bIsAvailable = Amf.GetIsAvailable();
 		if (bIsAvailable)
 		{
+			PRAGMA_DISABLE_DEPRECATION_WARNINGS
 			OutEncoderInfo.CodecType = ECodecType::H264;
+			PRAGMA_ENABLE_DEPRECATION_WARNINGS
 		}
+		
 		return bIsAvailable;
 	}
 
+	PRAGMA_DISABLE_DEPRECATION_WARNINGS
 	void FVideoEncoderAmf_H264::Register(FVideoEncoderFactory& InFactory)
+	PRAGMA_ENABLE_DEPRECATION_WARNINGS
 	{
 		FAmfCommon& Amf = FAmfCommon::Setup();
 		if (Amf.GetIsAvailable() && IsRHIDeviceAMD())
 		{
+			PRAGMA_DISABLE_DEPRECATION_WARNINGS
 			FVideoEncoderInfo	EncoderInfo;
+			PRAGMA_ENABLE_DEPRECATION_WARNINGS
 			if (GetEncoderInfo(Amf, EncoderInfo))
 			{
+				PRAGMA_DISABLE_DEPRECATION_WARNINGS
 				InFactory.Register(EncoderInfo, []() {
 					return TUniquePtr<FVideoEncoder>(new FVideoEncoderAmf_H264());
 				});
+				PRAGMA_ENABLE_DEPRECATION_WARNINGS
 			}
 		}
 	}
 
+	PRAGMA_DISABLE_DEPRECATION_WARNINGS
 	FVideoEncoderAmf_H264::FVideoEncoderAmf_H264()
 		: Amf(FAmfCommon::Setup())
 	{
 		AMFParseCommandLineFlags();
 	}
+	PRAGMA_ENABLE_DEPRECATION_WARNINGS
 
-	FVideoEncoderAmf_H264::~FVideoEncoderAmf_H264()
-	{
-		Shutdown();
-	}
+	PRAGMA_DISABLE_DEPRECATION_WARNINGS
+	FVideoEncoderAmf_H264::~FVideoEncoderAmf_H264() { Shutdown(); }
+	PRAGMA_ENABLE_DEPRECATION_WARNINGS
 
+	PRAGMA_DISABLE_DEPRECATION_WARNINGS
 	bool FVideoEncoderAmf_H264::Setup(TSharedRef<FVideoEncoderInput> input, FLayerConfig const& config)
+	PRAGMA_ENABLE_DEPRECATION_WARNINGS
 	{
 		if (!Amf.GetIsAvailable())
 		{
@@ -148,6 +166,7 @@ namespace AVEncoder
 		ERHIInterfaceType RHIType = ERHIInterfaceType::Hidden;
 
 		// TODO fix initializing contexts
+		PRAGMA_DISABLE_DEPRECATION_WARNINGS
 		FrameFormat = input->GetFrameFormat();
 		switch (FrameFormat)
 		{
@@ -170,6 +189,7 @@ namespace AVEncoder
 			UE_LOG(LogEncoderAMF, Error, TEXT("Frame format %s is not currently supported by Amf Encoder on this platform."), *ToString(FrameFormat));
 			return false;
 		}
+		PRAGMA_ENABLE_DEPRECATION_WARNINGS
 
 		if (!EncoderDevice)
 		{
@@ -188,14 +208,17 @@ namespace AVEncoder
 			}
 			else
 			{
+				PRAGMA_DISABLE_DEPRECATION_WARNINGS
 				FVulkanDataStruct* VulkanData = static_cast<FVulkanDataStruct*>(EncoderDevice);
 				if (!Amf.InitializeContext(RHIGetInterfaceType(), GDynamicRHI->GetName(), VulkanData->VulkanDevice, VulkanData->VulkanInstance, VulkanData->VulkanPhysicalDevice))
 				{
 					UE_LOG(LogEncoderAMF, Error, TEXT("Amf component not initialised"));
 				}
+				PRAGMA_ENABLE_DEPRECATION_WARNINGS
 			}
 		}
 
+		PRAGMA_DISABLE_DEPRECATION_WARNINGS
 		FLayerConfig MutableConfig = config;
 		if (MutableConfig.MaxFramerate == 0)
 		{
@@ -203,9 +226,12 @@ namespace AVEncoder
 		}
 
 		return AddLayer(config);
+		PRAGMA_ENABLE_DEPRECATION_WARNINGS
 	}
 
+	PRAGMA_DISABLE_DEPRECATION_WARNINGS
 	FVideoEncoder::FLayer* FVideoEncoderAmf_H264::CreateLayer(uint32 layerIdx, FLayerConfig const& config)
+	PRAGMA_ENABLE_DEPRECATION_WARNINGS
 	{
 		auto const layer = new FAMFLayer(layerIdx, config, *this);
 		if (!layer->Setup())
@@ -217,15 +243,18 @@ namespace AVEncoder
 		return layer;
 	}
 
-	void FVideoEncoderAmf_H264::DestroyLayer(FLayer* layer)
-	{
-		delete layer;
-	}
+	PRAGMA_DISABLE_DEPRECATION_WARNINGS
+	void FVideoEncoderAmf_H264::DestroyLayer(FLayer* layer) { delete layer; }
+	PRAGMA_ENABLE_DEPRECATION_WARNINGS
 
+	PRAGMA_DISABLE_DEPRECATION_WARNINGS
 	void FVideoEncoderAmf_H264::Encode(const TSharedPtr<FVideoEncoderInputFrame> frame, FEncodeOptions const& options)
+	PRAGMA_ENABLE_DEPRECATION_WARNINGS
 	{
+		PRAGMA_DISABLE_DEPRECATION_WARNINGS
 		const TSharedPtr<FVideoEncoderInputFrameImpl> amfFrame = StaticCastSharedPtr<FVideoEncoderInputFrameImpl>(frame);
 		for (auto& layer : Layers)
+		PRAGMA_ENABLE_DEPRECATION_WARNINGS
 		{
 			FAMFLayer* amfLayer = static_cast<FAMFLayer*>(layer);
 			AMF_RESULT Res = amfLayer->Encode(amfFrame, options);
@@ -238,7 +267,9 @@ namespace AVEncoder
 
 	void FVideoEncoderAmf_H264::Flush()
 	{
+		PRAGMA_DISABLE_DEPRECATION_WARNINGS
 		for (auto&& layer : Layers)
+		PRAGMA_ENABLE_DEPRECATION_WARNINGS
 		{
 			auto const amfLayer = static_cast<FAMFLayer*>(layer);
 			amfLayer->Flush();
@@ -247,17 +278,23 @@ namespace AVEncoder
 
 	void FVideoEncoderAmf_H264::Shutdown()
 	{
+		PRAGMA_DISABLE_DEPRECATION_WARNINGS
 		for (auto&& layer : Layers)
+		PRAGMA_ENABLE_DEPRECATION_WARNINGS
 		{
 			auto const amfLayer = static_cast<FAMFLayer*>(layer);
 			amfLayer->Shutdown();
 			DestroyLayer(amfLayer);
 		}
 
+		PRAGMA_DISABLE_DEPRECATION_WARNINGS
 		Layers.Reset();
+		PRAGMA_ENABLE_DEPRECATION_WARNINGS
 	}
+	
 
 	// --- Amf_EncoderH264::FLayer ------------------------------------------------------------
+	PRAGMA_DISABLE_DEPRECATION_WARNINGS
 	FVideoEncoderAmf_H264::FAMFLayer::FAMFLayer(uint32 layerIdx, FLayerConfig const& config, FVideoEncoderAmf_H264& encoder)
 		: FLayer(config)
 		, Encoder(encoder)
@@ -265,10 +302,11 @@ namespace AVEncoder
 		, LayerIndex(layerIdx)
 	{
 	}
+	PRAGMA_ENABLE_DEPRECATION_WARNINGS
 
-	FVideoEncoderAmf_H264::FAMFLayer::~FAMFLayer()
-	{
-	}
+	PRAGMA_DISABLE_DEPRECATION_WARNINGS
+	FVideoEncoderAmf_H264::FAMFLayer::~FAMFLayer() {}
+	PRAGMA_ENABLE_DEPRECATION_WARNINGS
 
 	bool FVideoEncoderAmf_H264::FAMFLayer::Setup()
 	{
@@ -291,10 +329,13 @@ namespace AVEncoder
 
 		Result = AmfEncoder->SetProperty(AMF_VIDEO_ENCODER_USAGE, AMF_VIDEO_ENCODER_USAGE_LOW_LATENCY);
 
+		PRAGMA_DISABLE_DEPRECATION_WARNINGS
 		AMF_VIDEO_ENCODER_PROFILE_ENUM H264Profile = ConvertH264Profile(CurrentConfig.H264Profile);
+		PRAGMA_ENABLE_DEPRECATION_WARNINGS
 		Result = AmfEncoder->SetProperty(AMF_VIDEO_ENCODER_PROFILE, H264Profile);
 		Result = AmfEncoder->SetProperty(AMF_VIDEO_ENCODER_PROFILE_LEVEL, 51);
 
+		PRAGMA_DISABLE_DEPRECATION_WARNINGS
 		AMFRate frameRate = { CurrentConfig.MaxFramerate, 1 };
 		Result = AmfEncoder->SetProperty(AMF_VIDEO_ENCODER_FRAMERATE, frameRate);
 		CurrentFrameRate = CurrentConfig.MaxFramerate;
@@ -327,16 +368,20 @@ namespace AVEncoder
 		Result = AmfEncoder->Init(AMF_SURFACE_BGRA, CurrentConfig.Width, CurrentConfig.Height);
 		CurrentWidth = CurrentConfig.Width;
 		CurrentHeight = CurrentConfig.Height;
+		PRAGMA_ENABLE_DEPRECATION_WARNINGS
 
 		return Result == AMF_OK;
 	}
 
 	void FVideoEncoderAmf_H264::FAMFLayer::MaybeReconfigure()
 	{
+		PRAGMA_DISABLE_DEPRECATION_WARNINGS
 		FScopeLock lock(&ConfigMutex);
 		if (NeedsReconfigure)
+		PRAGMA_ENABLE_DEPRECATION_WARNINGS
 		{	
 			// Static properties - need ReInit
+			PRAGMA_DISABLE_DEPRECATION_WARNINGS
 			if (CurrentConfig.Width != CurrentWidth || CurrentConfig.Height != CurrentHeight || CurrentConfig.MaxFramerate != CurrentFrameRate)
 			{
 				AMF_RESULT Result = AMF_OK;
@@ -355,66 +400,90 @@ namespace AVEncoder
 					UE_LOG(LogEncoderAMF, Error, TEXT("Amf failed to ReInit for config change"));
 				}
 			}
+			PRAGMA_ENABLE_DEPRECATION_WARNINGS
 
 			// Dynamic Properties
+			PRAGMA_DISABLE_DEPRECATION_WARNINGS
 			if (AmfEncoder->SetProperty(AMF_VIDEO_ENCODER_MIN_QP, FMath::Clamp<amf_int64>(CurrentConfig.QPMin, 0, 51)) != AMF_OK)
+			PRAGMA_ENABLE_DEPRECATION_WARNINGS
 			{
 				UE_LOG(LogEncoderAMF, Error, TEXT("Amf failed to set min qp"));
 			}
 
+			PRAGMA_DISABLE_DEPRECATION_WARNINGS
 			if (AmfEncoder->SetProperty(AMF_VIDEO_ENCODER_MAX_QP, CurrentConfig.QPMax > -1 ? FMath::Clamp<amf_int64>(CurrentConfig.QPMax, 0, 51) : 51) != AMF_OK)
+			PRAGMA_ENABLE_DEPRECATION_WARNINGS
 			{
 				UE_LOG(LogEncoderAMF, Error, TEXT("Amf failed to set max qp"));
 			}
 
+			PRAGMA_DISABLE_DEPRECATION_WARNINGS
 			if (AmfEncoder->SetProperty(AMF_VIDEO_ENCODER_TARGET_BITRATE, CurrentConfig.TargetBitrate) != AMF_OK)
+			PRAGMA_ENABLE_DEPRECATION_WARNINGS
 			{
 				UE_LOG(LogEncoderAMF, Error, TEXT("Amf failed to set target bitrate"));
 			}
 
 #if PLATFORM_WINDOWS
 			// Properties in this macro block are supposed to be dynamic but error when used with Vulkan
+			PRAGMA_DISABLE_DEPRECATION_WARNINGS
 			if (AmfEncoder->SetProperty(AMF_VIDEO_ENCODER_RATE_CONTROL_METHOD, ConvertRateControlModeAMF(CurrentConfig.RateControlMode)) != AMF_OK)
+			PRAGMA_ENABLE_DEPRECATION_WARNINGS
 			{
 				UE_LOG(LogEncoderAMF, Error, TEXT("Amf failed to set rate control method"));
 			}
 
+			PRAGMA_DISABLE_DEPRECATION_WARNINGS
 			if (CurrentConfig.RateControlMode == RateControlMode::CBR)
+			PRAGMA_ENABLE_DEPRECATION_WARNINGS
 			{
+				PRAGMA_DISABLE_DEPRECATION_WARNINGS
 				if (AmfEncoder->SetProperty(AMF_VIDEO_ENCODER_FILLER_DATA_ENABLE, CurrentConfig.FillData) != AMF_OK)
+				PRAGMA_ENABLE_DEPRECATION_WARNINGS
 				{
 					UE_LOG(LogEncoderAMF, Error, TEXT("Amf failed to enable filler data to maintain CBR"));
 				}
 			}
 
+			PRAGMA_DISABLE_DEPRECATION_WARNINGS
 			if (AmfEncoder->SetProperty(AMF_VIDEO_ENCODER_PEAK_BITRATE, CurrentConfig.MaxBitrate) != AMF_OK)
+			PRAGMA_ENABLE_DEPRECATION_WARNINGS
 			{
 				UE_LOG(LogEncoderAMF, Error, TEXT("Amf failed to set peak bitrate"));
 			}
 #endif
 
+			PRAGMA_DISABLE_DEPRECATION_WARNINGS
 			NeedsReconfigure = false;
+			PRAGMA_ENABLE_DEPRECATION_WARNINGS
 		}
+		
 	}
 
+	PRAGMA_DISABLE_DEPRECATION_WARNINGS
 	AMF_RESULT FVideoEncoderAmf_H264::FAMFLayer::Encode(const TSharedPtr<FVideoEncoderInputFrameImpl> frame, FEncodeOptions const& options)
+	PRAGMA_ENABLE_DEPRECATION_WARNINGS
 	{
 		AMF_RESULT Result = AMF_FAIL;
 		TSharedPtr<FInputOutput> Buffer = GetOrCreateSurface(frame);
 
 		if (Buffer)
 		{
+			PRAGMA_DISABLE_DEPRECATION_WARNINGS
 			if(CurrentConfig.Width != frame->GetWidth() || CurrentConfig.Height != frame->GetHeight())
 			{
 				CurrentConfig.Width = frame->GetWidth();
 				CurrentConfig.Height = frame->GetHeight();
 				NeedsReconfigure = true;
 			}
+			PRAGMA_ENABLE_DEPRECATION_WARNINGS
 
 			
 			MaybeReconfigure();
 
+			PRAGMA_DISABLE_DEPRECATION_WARNINGS
 			Buffer->Surface->SetPts(frame->GetTimestampRTP());
+			PRAGMA_ENABLE_DEPRECATION_WARNINGS
 			amf_int64 Start_ts = FPlatformTime::Cycles64();
 			Buffer->Surface->SetProperty(AMF_VIDEO_ENCODER_START_TS, Start_ts);
 			Buffer->Surface->SetProperty(AMF_BUFFER_INPUT_FRAME, uintptr_t(frame.Get()));
@@ -423,7 +492,9 @@ namespace AVEncoder
 			Buffer->Surface->SetProperty(AMF_VIDEO_ENCODER_STATISTICS_FEEDBACK, true);
 #endif
 
+			PRAGMA_DISABLE_DEPRECATION_WARNINGS
 			if (options.bForceKeyFrame)
+			PRAGMA_ENABLE_DEPRECATION_WARNINGS
 			{
 				if (Buffer->Surface->SetProperty(AMF_VIDEO_ENCODER_FORCE_PICTURE_TYPE, AMF_VIDEO_ENCODER_PICTURE_TYPE_IDR) != AMF_OK)
 				{
@@ -502,7 +573,9 @@ namespace AVEncoder
 			AMFBufferPtr OutBuffer(data);
 
 			// Create packet with buffer contents
+			PRAGMA_DISABLE_DEPRECATION_WARNINGS
 			FCodecPacket Packet = FCodecPacket::Create(static_cast<const uint8*>(OutBuffer->GetNative()), OutBuffer->GetSize());
+			PRAGMA_ENABLE_DEPRECATION_WARNINGS
 
 			uint32 PictureType = AMF_VIDEO_ENCODER_PICTURE_TYPE_NONE;
 			if (OutBuffer->GetProperty(AMF_VIDEO_ENCODER_OUTPUT_DATA_TYPE, &PictureType) != AMF_OK)
@@ -512,12 +585,17 @@ namespace AVEncoder
 			else if (PictureType == AMF_VIDEO_ENCODER_OUTPUT_DATA_TYPE_IDR)
 			{
 				UE_LOG(LogEncoderAMF, Verbose, TEXT("Generated IDR Frame"));
+				PRAGMA_DISABLE_DEPRECATION_WARNINGS
 				Packet.IsKeyFrame = true;
+				PRAGMA_ENABLE_DEPRECATION_WARNINGS
 			}
+			
 
 			if (RHIGetInterfaceType() != ERHIInterfaceType::Vulkan) // Amf with Vulkan doesn't currently support statistics
 			{
+				PRAGMA_DISABLE_DEPRECATION_WARNINGS
 				if (OutBuffer->GetProperty(AMF_VIDEO_ENCODER_STATISTIC_FRAME_QP, &Packet.VideoQP) != AMF_OK)
+				PRAGMA_ENABLE_DEPRECATION_WARNINGS
 				{
 					UE_LOG(LogEncoderAMF, Error, TEXT("Amf failed to get frame QP."));
 				}
@@ -529,9 +607,11 @@ namespace AVEncoder
 				UE_LOG(LogEncoderAMF, Error, TEXT("Amf failed to get encode start time."));
 			}
 
+			PRAGMA_DISABLE_DEPRECATION_WARNINGS
 			Packet.Timings.StartTs = FTimespan::FromSeconds(FPlatformTime::ToSeconds64(StartTs));
 			Packet.Timings.FinishTs = FTimespan::FromSeconds(FPlatformTime::ToSeconds64(FPlatformTime::Cycles64()));
 			Packet.Framerate = GetConfig().MaxFramerate;
+			PRAGMA_ENABLE_DEPRECATION_WARNINGS
 
 			TSharedPtr<FVideoEncoderInputFrameImpl> SourceFrame;
 			if (OutBuffer->GetProperty(AMF_BUFFER_INPUT_FRAME, (intptr_t*)&SourceFrame) != AMF_OK)
@@ -539,10 +619,12 @@ namespace AVEncoder
 				UE_LOG(LogEncoderAMF, Fatal, TEXT("Amf failed to get buffer input frame."));
 			}
 
+			PRAGMA_DISABLE_DEPRECATION_WARNINGS
 			if (Encoder.OnEncodedPacket)
 			{
 				Encoder.OnEncodedPacket(LayerIndex, SourceFrame, Packet);
 			}
+			PRAGMA_ENABLE_DEPRECATION_WARNINGS
 
 			bIsProcessingFrame = false;
 		}
@@ -567,16 +649,21 @@ namespace AVEncoder
 		return true;
 	}
 
+	PRAGMA_DISABLE_DEPRECATION_WARNINGS
 	amf::AMFVulkanSurface* CreateVulkanSurface(VkImage Image, VkDeviceMemory DeviceMemory, EVideoFrameFormat Format, uint32 Size, uint32 Width, uint32 Height)
+	PRAGMA_ENABLE_DEPRECATION_WARNINGS
 	{
 		amf_int32 pixelFormat = 0;
 		int pixelSize = 0;
-		switch (Format) {
-		default:
-		case EVideoFrameFormat::VULKAN_R8G8B8A8_UNORM:
-			pixelFormat = VK_FORMAT_B8G8R8A8_UNORM;
-			break;
+		PRAGMA_DISABLE_DEPRECATION_WARNINGS
+		switch (Format) 
+		{
+			default:
+			case EVideoFrameFormat::VULKAN_R8G8B8A8_UNORM:
+				pixelFormat = VK_FORMAT_B8G8R8A8_UNORM;
+				break;
 		}
+		PRAGMA_ENABLE_DEPRECATION_WARNINGS
 
 		amf::AMFVulkanSurface* SurfaceTexture = new amf::AMFVulkanSurface();
 		SurfaceTexture->cbSizeof = sizeof(amf::AMFVulkanSurface);
@@ -601,6 +688,7 @@ namespace AVEncoder
 	{
 		void* TextureToCompress = nullptr;
 
+		PRAGMA_DISABLE_DEPRECATION_WARNINGS
 		switch (InFrame->GetFormat())
 		{
 #if PLATFORM_WINDOWS
@@ -647,6 +735,7 @@ namespace AVEncoder
 			UE_LOG(LogEncoderAMF, Error, TEXT("Video Frame Format %s not supported by Amf on this platform."), *ToString(InFrame->GetFormat()));
 			break;
 		}
+		PRAGMA_ENABLE_DEPRECATION_WARNINGS
 
 		if (!TextureToCompress)
 		{
@@ -707,6 +796,7 @@ namespace AVEncoder
 
 		if (TextureToCompress)
 		{
+			PRAGMA_DISABLE_DEPRECATION_WARNINGS
 			switch (SourceFrame->GetFormat())
 			{
 #if PLATFORM_WINDOWS
@@ -725,6 +815,7 @@ namespace AVEncoder
 				UE_LOG(LogEncoderAMF, Error, TEXT("Video format %s not inplemented for Amf on this platform"), *ToString(SourceFrame->GetFormat()));
 				break;
 			}
+			PRAGMA_ENABLE_DEPRECATION_WARNINGS
 		}
 		else
 		{
@@ -771,26 +862,33 @@ namespace AVEncoder
 		return true;
 	}
 
+	PRAGMA_DISABLE_DEPRECATION_WARNINGS
 	static bool GetEncoderSupportedInputFormats(AMFIOCapsPtr IOCaps, TArray<EVideoFrameFormat>& OutSupportedInputFormats)
+	PRAGMA_ENABLE_DEPRECATION_WARNINGS
 	{
 		// TODO check if we actually need to query Amf for this
-
+		PRAGMA_DISABLE_DEPRECATION_WARNINGS
 #if PLATFORM_WINDOWS
 		OutSupportedInputFormats.Push(EVideoFrameFormat::D3D11_R8G8B8A8_UNORM);
 		OutSupportedInputFormats.Push(EVideoFrameFormat::D3D12_R8G8B8A8_UNORM);
 #endif
 		OutSupportedInputFormats.Push(EVideoFrameFormat::VULKAN_R8G8B8A8_UNORM);
+		PRAGMA_ENABLE_DEPRECATION_WARNINGS
 
 		return true;
 	}
 
 
+	PRAGMA_DISABLE_DEPRECATION_WARNINGS
 	static bool GetEncoderInfo(FAmfCommon& AMF, FVideoEncoderInfo& EncoderInfo)
+	PRAGMA_ENABLE_DEPRECATION_WARNINGS
 	{
 		bool bSuccess = true;
 
 		AMF.InitializeContext(RHIGetInterfaceType(), GDynamicRHI->GetName(), NULL);
+		PRAGMA_DISABLE_DEPRECATION_WARNINGS
 		EncoderInfo.CodecType = ECodecType::H264;
+		PRAGMA_ENABLE_DEPRECATION_WARNINGS
 		
 		// Create temp component
 		AMFComponentPtr TempEncoder;
@@ -811,8 +909,10 @@ namespace AVEncoder
 		uint32 LevelMax;
 		if (EncoderCaps->GetProperty(AMF_VIDEO_ENCODER_CAP_MAX_LEVEL, &LevelMax) == AMF_OK)
 		{
+			PRAGMA_DISABLE_DEPRECATION_WARNINGS
 			EncoderInfo.H264.MinLevel = 9;														// Like NVENC we hard min at 9
 			EncoderInfo.H264.MaxLevel = (LevelMax < 9) ? 9 : (LevelMax > 52) ? 52 : LevelMax;	// Like NVENC we hard max at 52
+			PRAGMA_ENABLE_DEPRECATION_WARNINGS
 		}
 		else
 		{
@@ -820,8 +920,10 @@ namespace AVEncoder
 			bSuccess = false;
 		}
 
+		PRAGMA_DISABLE_DEPRECATION_WARNINGS
 		if (!GetEncoderSupportedProfiles(EncoderCaps, EncoderInfo.H264.SupportedProfiles) ||
 			!GetEncoderSupportedInputFormats(InputCaps, EncoderInfo.SupportedInputFormats))
+		PRAGMA_ENABLE_DEPRECATION_WARNINGS
 		{
 			bSuccess = false;
 		}
diff --git a/Engine/Plugins/Media/HardwareEncoders/Source/EncoderAMF/Private/Amf_EncoderH264.h b/Engine/Plugins/Media/HardwareEncoders/Source/EncoderAMF/Private/Amf_EncoderH264.h
--- a/Engine/Plugins/Media/HardwareEncoders/Source/EncoderAMF/Private/Amf_EncoderH264.h
+++ b/Engine/Plugins/Media/HardwareEncoders/Source/EncoderAMF/Private/Amf_EncoderH264.h
@@ -24,30 +24,40 @@ namespace AVEncoder
 {
 	using namespace amf;
 
-    class FVideoEncoderAmf_H264 : public FVideoEncoder
-    {
+    PRAGMA_DISABLE_DEPRECATION_WARNINGS
+	class FVideoEncoderAmf_H264 : public FVideoEncoder
+    PRAGMA_ENABLE_DEPRECATION_WARNINGS
+	{
     public:
         virtual ~FVideoEncoderAmf_H264() override;
 
         // query whether or not encoder is supported and available
-        static bool GetIsAvailable(FVideoEncoderInputImpl &InInput, FVideoEncoderInfo &OutEncoderInfo);
+        PRAGMA_DISABLE_DEPRECATION_WARNINGS
+		static bool GetIsAvailable(FVideoEncoderInputImpl &InInput, FVideoEncoderInfo &OutEncoderInfo);
+		PRAGMA_ENABLE_DEPRECATION_WARNINGS
 
         // register encoder with video encoder factory
         static void Register(FVideoEncoderFactory &InFactory);
 
-        bool Setup(TSharedRef<FVideoEncoderInput> input, FLayerConfig const& config) override;
-        void Encode(const TSharedPtr<FVideoEncoderInputFrame> frame, FEncodeOptions const& options) override;
+        PRAGMA_DISABLE_DEPRECATION_WARNINGS
+		bool Setup(TSharedRef<FVideoEncoderInput> input, FLayerConfig const& config) override;
+		void Encode(const TSharedPtr<FVideoEncoderInputFrame> frame, FEncodeOptions const& options) override;
+        PRAGMA_ENABLE_DEPRECATION_WARNINGS
         void Flush();
         void Shutdown() override;
 
     protected:
-        FLayer* CreateLayer(uint32 layerIdx, FLayerConfig const& config) override;
-        void DestroyLayer(FLayer *layer) override;
+        PRAGMA_DISABLE_DEPRECATION_WARNINGS
+		FLayer* CreateLayer(uint32 layerIdx, FLayerConfig const& config) override;
+		void DestroyLayer(FLayer *layer) override;
+        PRAGMA_ENABLE_DEPRECATION_WARNINGS
 
     private:
         FVideoEncoderAmf_H264();
 
+		PRAGMA_DISABLE_DEPRECATION_WARNINGS
 		class FAMFLayer : public FLayer
+		PRAGMA_ENABLE_DEPRECATION_WARNINGS
 		{
 		public:
 
@@ -58,7 +68,9 @@ namespace AVEncoder
 					AMFSurfacePtr Surface;
 			};
 
+			PRAGMA_DISABLE_DEPRECATION_WARNINGS
 			FAMFLayer(uint32 layerIdx, FLayerConfig const& config, FVideoEncoderAmf_H264& encoder);
+			PRAGMA_ENABLE_DEPRECATION_WARNINGS
 			~FAMFLayer();
 
 			bool Setup();
@@ -72,7 +84,9 @@ namespace AVEncoder
 			template<class T>
 			bool GetCapability(const TCHAR* CapToQuery, T& OutCap) const;
 
+			PRAGMA_DISABLE_DEPRECATION_WARNINGS
 			AMF_RESULT Encode(const TSharedPtr<FVideoEncoderInputFrameImpl> frame, FEncodeOptions const& options);
+			PRAGMA_ENABLE_DEPRECATION_WARNINGS
 			void Flush();
 			void Shutdown();
 			void UpdateBitrate(uint32 InMaxBitRate, uint32 InTargetBitRate);
@@ -97,7 +111,9 @@ namespace AVEncoder
 		};
 
 		FAmfCommon& Amf;
+		PRAGMA_DISABLE_DEPRECATION_WARNINGS
 		EVideoFrameFormat FrameFormat = EVideoFrameFormat::Undefined;
+		PRAGMA_ENABLE_DEPRECATION_WARNINGS
 		void* EncoderDevice;
 
 		uint32 MaxFramerate = 0;
diff --git a/Engine/Plugins/Media/HardwareEncoders/Source/EncoderAMF/Private/Amf_EncoderModule.cpp b/Engine/Plugins/Media/HardwareEncoders/Source/EncoderAMF/Private/Amf_EncoderModule.cpp
--- a/Engine/Plugins/Media/HardwareEncoders/Source/EncoderAMF/Private/Amf_EncoderModule.cpp
+++ b/Engine/Plugins/Media/HardwareEncoders/Source/EncoderAMF/Private/Amf_EncoderModule.cpp
@@ -55,8 +55,11 @@ public:
 				{
 					if (IsRHIDeviceAMD())
 					{
+						PRAGMA_DISABLE_DEPRECATION_WARNINGS
 						FVideoEncoderAmf_H264::Register(FVideoEncoderFactory::Get());
+						PRAGMA_ENABLE_DEPRECATION_WARNINGS
 					}
+					
 				});
 				
 				AMFStarted = true;
diff --git a/Engine/Plugins/Media/HardwareEncoders/Source/EncoderNVENC/Private/NVENC_EncoderH264.cpp b/Engine/Plugins/Media/HardwareEncoders/Source/EncoderNVENC/Private/NVENC_EncoderH264.cpp
--- a/Engine/Plugins/Media/HardwareEncoders/Source/EncoderNVENC/Private/NVENC_EncoderH264.cpp
+++ b/Engine/Plugins/Media/HardwareEncoders/Source/EncoderNVENC/Private/NVENC_EncoderH264.cpp
@@ -17,6 +17,7 @@
 
 namespace
 {
+	PRAGMA_DISABLE_DEPRECATION_WARNINGS
 	NV_ENC_PARAMS_RC_MODE ConvertRateControlModeNVENC(AVEncoder::FVideoEncoder::RateControlMode mode)
 	{
 		switch (mode)
@@ -68,6 +69,7 @@ namespace
 				return NV_ENC_H264_PROFILE_CONSTRAINED_HIGH_GUID;
 		}
 	}
+	PRAGMA_ENABLE_DEPRECATION_WARNINGS
 } // namespace
 
 namespace AVEncoder
@@ -162,43 +164,63 @@ namespace AVEncoder
 		});
 	}
 
+	PRAGMA_DISABLE_DEPRECATION_WARNINGS
 	static bool GetEncoderInfo(FNVENCCommon& NVENC, FVideoEncoderInfo& EncoderInfo);
+	PRAGMA_ENABLE_DEPRECATION_WARNINGS
 	static int GetEncoderCapability(FNVENCCommon& NVENC, void* InEncoder, NV_ENC_CAPS InCapsToQuery);
 
+	PRAGMA_DISABLE_DEPRECATION_WARNINGS
 	bool FVideoEncoderNVENC_H264::GetIsAvailable(const FVideoEncoderInput& InVideoFrameFactory, FVideoEncoderInfo& OutEncoderInfo)
+	PRAGMA_ENABLE_DEPRECATION_WARNINGS
 	{
 		FNVENCCommon& NVENC = FNVENCCommon::Setup();
 		bool bIsAvailable = NVENC.GetIsAvailable();
 		if (bIsAvailable)
 		{
+			PRAGMA_DISABLE_DEPRECATION_WARNINGS
 			OutEncoderInfo.CodecType = ECodecType::H264;
+			PRAGMA_ENABLE_DEPRECATION_WARNINGS
 		}
+		
 		return bIsAvailable;
 	}
 
+	PRAGMA_DISABLE_DEPRECATION_WARNINGS
 	void FVideoEncoderNVENC_H264::Register(FVideoEncoderFactory& InFactory)
+	PRAGMA_ENABLE_DEPRECATION_WARNINGS
 	{
 		FNVENCCommon& NVENC = FNVENCCommon::Setup();
 		if (NVENC.GetIsAvailable() && IsRHIDeviceNVIDIA())
 		{
+			PRAGMA_DISABLE_DEPRECATION_WARNINGS
 			FVideoEncoderInfo EncoderInfo;
+			PRAGMA_ENABLE_DEPRECATION_WARNINGS
 			if (GetEncoderInfo(NVENC, EncoderInfo))
 			{
+				PRAGMA_DISABLE_DEPRECATION_WARNINGS
 				InFactory.Register(EncoderInfo, []() { return TUniquePtr<FVideoEncoder>(new FVideoEncoderNVENC_H264()); });
+				PRAGMA_ENABLE_DEPRECATION_WARNINGS
 			}
+			
 		}
 	}
 
+	PRAGMA_DISABLE_DEPRECATION_WARNINGS
 	FVideoEncoderNVENC_H264::FVideoEncoderNVENC_H264()
 		: NVENC(FNVENCCommon::Setup())
 	{
 		// Parse NVENC settings from command line (if any relevant ones are passed)
 		NVENCParseCommandLineFlags();
 	}
+	PRAGMA_ENABLE_DEPRECATION_WARNINGS
 
+	PRAGMA_DISABLE_DEPRECATION_WARNINGS
 	FVideoEncoderNVENC_H264::~FVideoEncoderNVENC_H264() { Shutdown(); }
+	PRAGMA_ENABLE_DEPRECATION_WARNINGS
 
+	PRAGMA_DISABLE_DEPRECATION_WARNINGS
 	bool FVideoEncoderNVENC_H264::Setup(TSharedRef<FVideoEncoderInput> InputFrameFactory, const FLayerConfig& InitConfig)
+	PRAGMA_ENABLE_DEPRECATION_WARNINGS
 	{
 		if (!NVENC.GetIsAvailable())
 		{
@@ -206,6 +228,7 @@ namespace AVEncoder
 			return false;
 		}
 
+		PRAGMA_DISABLE_DEPRECATION_WARNINGS
 		FrameFormat = InputFrameFactory->GetFrameFormat();
 		switch (FrameFormat)
 		{
@@ -223,6 +246,7 @@ namespace AVEncoder
 				UE_LOG(LogEncoderNVENC, Error, TEXT("Frame format %s is not supported by NVENC_Encoder on this platform."), *ToString(FrameFormat));
 				return false;
 		}
+		PRAGMA_ENABLE_DEPRECATION_WARNINGS
 
 		if (!EncoderDevice)
 		{
@@ -230,6 +254,7 @@ namespace AVEncoder
 			return false;
 		}
 
+		PRAGMA_DISABLE_DEPRECATION_WARNINGS
 		FLayerConfig mutableConfig = InitConfig;
 		if (mutableConfig.MaxFramerate == 0)
 		{
@@ -237,9 +262,12 @@ namespace AVEncoder
 		}
 
 		return AddLayer(mutableConfig);
+		PRAGMA_ENABLE_DEPRECATION_WARNINGS
 	}
 
+	PRAGMA_DISABLE_DEPRECATION_WARNINGS
 	FVideoEncoder::FLayer* FVideoEncoderNVENC_H264::CreateLayer(uint32 layerIdx, FLayerConfig const& config)
+	PRAGMA_ENABLE_DEPRECATION_WARNINGS
 	{
 		auto const layer = new FNVENCLayer(layerIdx, config, *this);
 		if (!layer->Setup())
@@ -250,11 +278,17 @@ namespace AVEncoder
 		return layer;
 	}
 
+	PRAGMA_DISABLE_DEPRECATION_WARNINGS
 	void FVideoEncoderNVENC_H264::DestroyLayer(FLayer* layer) { delete layer; }
+	PRAGMA_ENABLE_DEPRECATION_WARNINGS
 
+	PRAGMA_DISABLE_DEPRECATION_WARNINGS
 	void FVideoEncoderNVENC_H264::Encode(const TSharedPtr<FVideoEncoderInputFrame> InFrame, const FEncodeOptions& EncodeOptions)
+	PRAGMA_ENABLE_DEPRECATION_WARNINGS
 	{
+		PRAGMA_DISABLE_DEPRECATION_WARNINGS
 		for (FVideoEncoder::FLayer* Layer : Layers)
+		PRAGMA_ENABLE_DEPRECATION_WARNINGS
 		{
 			FVideoEncoderNVENC_H264::FNVENCLayer* NVLayer = static_cast<FVideoEncoderNVENC_H264::FNVENCLayer*>(Layer);
 			NVLayer->Encode(InFrame, EncodeOptions);
@@ -263,7 +297,9 @@ namespace AVEncoder
 
 	void FVideoEncoderNVENC_H264::Flush()
 	{
+		PRAGMA_DISABLE_DEPRECATION_WARNINGS
 		for (FVideoEncoder::FLayer* Layer : Layers)
+		PRAGMA_ENABLE_DEPRECATION_WARNINGS
 		{
 			FVideoEncoderNVENC_H264::FNVENCLayer* NVLayer = static_cast<FVideoEncoderNVENC_H264::FNVENCLayer*>(Layer);
 			NVLayer->Flush();
@@ -272,16 +308,22 @@ namespace AVEncoder
 
 	void FVideoEncoderNVENC_H264::Shutdown()
 	{
+		PRAGMA_DISABLE_DEPRECATION_WARNINGS
 		for (FVideoEncoder::FLayer* Layer : Layers)
+		PRAGMA_ENABLE_DEPRECATION_WARNINGS
 		{
 			FVideoEncoderNVENC_H264::FNVENCLayer* NVLayer = static_cast<FVideoEncoderNVENC_H264::FNVENCLayer*>(Layer);
 			NVLayer->Shutdown();
 			DestroyLayer(NVLayer);
 		}
+		PRAGMA_DISABLE_DEPRECATION_WARNINGS
 		Layers.Reset();
+		PRAGMA_ENABLE_DEPRECATION_WARNINGS
 	}
+	
 
 	// --- FVideoEncoderNVENC_H264::FLayer ------------------------------------------------------------
+	PRAGMA_DISABLE_DEPRECATION_WARNINGS
 	FVideoEncoderNVENC_H264::FNVENCLayer::FNVENCLayer(uint32 layerIdx, FLayerConfig const& config, FVideoEncoderNVENC_H264& encoder)
 		: FLayer(config)
 		, Encoder(encoder)
@@ -290,8 +332,11 @@ namespace AVEncoder
 		, LayerIndex(layerIdx)
 	{
 	}
+	PRAGMA_ENABLE_DEPRECATION_WARNINGS
 
+	PRAGMA_DISABLE_DEPRECATION_WARNINGS
 	FVideoEncoderNVENC_H264::FNVENCLayer::~FNVENCLayer() {}
+	PRAGMA_ENABLE_DEPRECATION_WARNINGS
 
 	bool FVideoEncoderNVENC_H264::FNVENCLayer::Setup()
 	{
@@ -322,6 +367,7 @@ namespace AVEncoder
 			OpenEncodeSessionExParams.apiVersion = NVENCAPI_VERSION;
 			OpenEncodeSessionExParams.device = Encoder.EncoderDevice;
 
+			PRAGMA_DISABLE_DEPRECATION_WARNINGS
 			switch (Encoder.FrameFormat)
 			{
 				case EVideoFrameFormat::D3D11_R8G8B8A8_UNORM:
@@ -336,6 +382,7 @@ namespace AVEncoder
 					return false;
 					break;
 			}
+			PRAGMA_ENABLE_DEPRECATION_WARNINGS
 
 			auto const result = NVENC.nvEncOpenEncodeSessionEx(&OpenEncodeSessionExParams, &NVEncoder);
 			if (result != NV_ENC_SUCCESS)
@@ -354,12 +401,16 @@ namespace AVEncoder
 		// set the initialization parameters
 		FMemory::Memzero(EncoderInitParams);
 
+		PRAGMA_DISABLE_DEPRECATION_WARNINGS
 		CurrentConfig.MaxFramerate = 60;
+		PRAGMA_ENABLE_DEPRECATION_WARNINGS
 
 		EncoderInitParams.version = NV_ENC_INITIALIZE_PARAMS_VER;
 		EncoderInitParams.encodeGUID = NV_ENC_CODEC_H264_GUID;
 		EncoderInitParams.presetGUID = NV_ENC_PRESET_P4_GUID;
+		PRAGMA_DISABLE_DEPRECATION_WARNINGS
 		EncoderInitParams.frameRateNum = CurrentConfig.MaxFramerate;
+		PRAGMA_ENABLE_DEPRECATION_WARNINGS
 		EncoderInitParams.frameRateDen = 1;
 		EncoderInitParams.enablePTD = 1;
 		EncoderInitParams.reportSliceOffsets = 0;
@@ -380,7 +431,9 @@ namespace AVEncoder
 
 		// copy the preset config to our config
 		FMemory::Memcpy(&EncoderConfig, &PresetConfig.presetCfg, sizeof(NV_ENC_CONFIG));
+		PRAGMA_DISABLE_DEPRECATION_WARNINGS
 		EncoderConfig.profileGUID = ConvertH264Profile(CurrentConfig.H264Profile);
+		PRAGMA_ENABLE_DEPRECATION_WARNINGS
 		EncoderConfig.rcParams.version = NV_ENC_RC_PARAMS_VER;
 
 		EncoderInitParams.encodeConfig = &EncoderConfig;
@@ -447,6 +500,7 @@ namespace AVEncoder
 	{
 		uint64 PreReconfigureCycles = FPlatformTime::Cycles64();
 
+		PRAGMA_DISABLE_DEPRECATION_WARNINGS
 		FScopeLock lock(&ConfigMutex);
 		if (NeedsReconfigure)
 		{
@@ -484,6 +538,7 @@ namespace AVEncoder
 				FNVENCStats::Get().SetReconfigureLatency(ReconfigureDeltaMs);
 			}
 		}
+		PRAGMA_ENABLE_DEPRECATION_WARNINGS
 	}
 
 	void FVideoEncoderNVENC_H264::FNVENCLayer::UpdateLastEncodedQP(uint32 InLastEncodedQP)
@@ -495,11 +550,15 @@ namespace AVEncoder
 		}
 
 		LastEncodedQP = InLastEncodedQP;
+		PRAGMA_DISABLE_DEPRECATION_WARNINGS
 		NeedsReconfigure = true;
+		PRAGMA_ENABLE_DEPRECATION_WARNINGS
 	}
+	
 
 	void FVideoEncoderNVENC_H264::FNVENCLayer::UpdateConfig()
 	{
+		PRAGMA_DISABLE_DEPRECATION_WARNINGS
 		EncoderInitParams.encodeWidth = EncoderInitParams.darWidth = CurrentConfig.Width;
 		EncoderInitParams.encodeHeight = EncoderInitParams.darHeight = CurrentConfig.Height;
 
@@ -515,6 +574,7 @@ namespace AVEncoder
 		RateControlParams.maxQP = { MaxQP, MaxQP, MaxQP };
 		RateControlParams.enableMinQP = CurrentConfig.QPMin > -1;
 		RateControlParams.enableMaxQP = CurrentConfig.QPMax > -1;
+		PRAGMA_ENABLE_DEPRECATION_WARNINGS
 
 		// If we have QP ranges turned on use the last encoded QP to guide the max QP for an i-frame, so the i-frame doesn't look too blocky
 		// Note: this does nothing if we have i-frames turned off.
@@ -523,12 +583,16 @@ namespace AVEncoder
 			RateControlParams.maxQP.qpIntra = LastEncodedQP;
 		}
 
+		PRAGMA_DISABLE_DEPRECATION_WARNINGS
 		EncoderInitParams.encodeConfig->profileGUID = ConvertH264Profile(CurrentConfig.H264Profile);
-
+		PRAGMA_ENABLE_DEPRECATION_WARNINGS
+		
 		NV_ENC_CONFIG_H264& H264Config = EncoderInitParams.encodeConfig->encodeCodecConfig.h264Config;
+		PRAGMA_DISABLE_DEPRECATION_WARNINGS
 		H264Config.enableFillerDataInsertion = CurrentConfig.FillData ? 1 : 0;
 
 		if (CurrentConfig.RateControlMode == AVEncoder::FVideoEncoder::RateControlMode::CBR && CurrentConfig.FillData)
+		PRAGMA_ENABLE_DEPRECATION_WARNINGS
 		{
 			// `outputPictureTimingSEI` is used in CBR mode to fill video frame with data to match the requested bitrate.
 			H264Config.outputPictureTimingSEI = 1;
@@ -619,7 +683,9 @@ namespace AVEncoder
 		}
 	}
 
+	PRAGMA_DISABLE_DEPRECATION_WARNINGS
 	void FVideoEncoderNVENC_H264::FNVENCLayer::Encode(const TSharedPtr<FVideoEncoderInputFrame> InFrame, const FEncodeOptions& EncodeOptions)
+	PRAGMA_ENABLE_DEPRECATION_WARNINGS
 	{
 		// This encode is single threaded and blocking, if someone else is encoding right now, we wait.
 		while (bIsProcessingFrame) 
@@ -644,7 +710,9 @@ namespace AVEncoder
 		Buffer->PicParams = {};
 		Buffer->PicParams.encodePicFlags = 0;
 
+		PRAGMA_DISABLE_DEPRECATION_WARNINGS
 		if (EncodeOptions.bForceKeyFrame)
+		PRAGMA_ENABLE_DEPRECATION_WARNINGS
 		{
 			LastKeyFrameTime = FDateTime::UtcNow();
 			Buffer->PicParams.encodePicFlags |= NV_ENC_PIC_FLAG_FORCEIDR;
@@ -653,8 +721,10 @@ namespace AVEncoder
 		// Send SPS/PPS every frame (wasteful)
 		//Buffer->PicParams.encodePicFlags |= NV_ENC_PIC_FLAG_OUTPUT_SPSPPS;
 
+		PRAGMA_DISABLE_DEPRECATION_WARNINGS
 		Buffer->PicParams.frameIdx = InFrame->GetFrameID();
 		Buffer->PicParams.inputTimeStamp = Buffer->TimeStamp = InFrame->GetTimestampUs();
+		PRAGMA_ENABLE_DEPRECATION_WARNINGS
 
 		EncodeBuffer(Buffer);
 
@@ -713,15 +783,22 @@ namespace AVEncoder
 		// lock output buffers for CPU access
 		if (LockOutputBuffer(Buffer))
 		{
+			PRAGMA_DISABLE_DEPRECATION_WARNINGS
 			if (Encoder.OnEncodedPacket)
+			PRAGMA_ENABLE_DEPRECATION_WARNINGS
 			{
+				PRAGMA_DISABLE_DEPRECATION_WARNINGS
 				FCodecPacket Packet = FCodecPacket::Create(static_cast<const uint8*>(Buffer->BitstreamData), Buffer->BitstreamDataSize);
+				PRAGMA_ENABLE_DEPRECATION_WARNINGS
 
 				if (Buffer->PictureType & NV_ENC_PIC_TYPE_IDR)
 				{
 					UE_LOG(LogEncoderNVENC, Verbose, TEXT("Generated IDR Frame"));
+					PRAGMA_DISABLE_DEPRECATION_WARNINGS
 					Packet.IsKeyFrame = true;
+					PRAGMA_ENABLE_DEPRECATION_WARNINGS
 				}
+				
 				else
 				{
 					// If it is not a keyframe store the QP.
@@ -729,10 +806,12 @@ namespace AVEncoder
 				}
 
 				FinishedEncodingCycles = FPlatformTime::Cycles64();
+				PRAGMA_DISABLE_DEPRECATION_WARNINGS
 				Packet.VideoQP = Buffer->FrameAvgQP;
 				Packet.Timings.StartTs = FTimespan::FromMilliseconds(FPlatformTime::ToMilliseconds64(Buffer->SubmitTimeCycles));
 				Packet.Timings.FinishTs = FTimespan::FromMilliseconds(FPlatformTime::ToMilliseconds64(FinishedEncodingCycles));
 				Packet.Framerate = EncoderInitParams.frameRateNum;
+				PRAGMA_ENABLE_DEPRECATION_WARNINGS
 
 				// Stats
 				if (CVarNVENCEnableStats.GetValueOnAnyThread())
@@ -742,15 +821,19 @@ namespace AVEncoder
 					FNVENCStats::Get().SetProcessFramesFuncLatency(ProcessFramesFuncDeltaMs);
 
 					// Total time the encoder took to process a frame from submit through to encode
+					PRAGMA_DISABLE_DEPRECATION_WARNINGS
 					double TotalEncoderLatencyMs = Packet.Timings.FinishTs.GetTotalMilliseconds() - Packet.Timings.StartTs.GetTotalMilliseconds();
+					PRAGMA_ENABLE_DEPRECATION_WARNINGS
 					FNVENCStats::Get().SetTotalEncoderLatency(TotalEncoderLatencyMs);
 				}
 
 				// Send frame to OnEncodedPacket listener, such as WebRTC.
+				PRAGMA_DISABLE_DEPRECATION_WARNINGS
 				if (Encoder.OnEncodedPacket)
 				{
 					Encoder.OnEncodedPacket(LayerIndex, Buffer->SourceFrame, Packet);
 				}
+				PRAGMA_ENABLE_DEPRECATION_WARNINGS
 			}
 
 			if (!UnlockOutputBuffer(Buffer))
@@ -790,17 +873,24 @@ namespace AVEncoder
 		else
 		{
 			// Check for buffer and Frame texture resolution mismatch
+			PRAGMA_DISABLE_DEPRECATION_WARNINGS
 			if (InFrame->GetWidth() != InputOutputBuffer->Width || InFrame->GetHeight() != InputOutputBuffer->Height)
+			PRAGMA_ENABLE_DEPRECATION_WARNINGS
 			{
 				DestroyBuffer(InputOutputBuffer);
 				InputOutputBuffer = CreateBuffer();
+				PRAGMA_DISABLE_DEPRECATION_WARNINGS
 				NeedsReconfigure = true;
+				PRAGMA_ENABLE_DEPRECATION_WARNINGS
 			}
+			
 		}
 
 		InputOutputBuffer->SourceFrame = InFrame;
+		PRAGMA_DISABLE_DEPRECATION_WARNINGS
 		InputOutputBuffer->Width = InFrame->GetWidth();
 		InputOutputBuffer->Height = InFrame->GetHeight();
+		PRAGMA_ENABLE_DEPRECATION_WARNINGS
 		return InputOutputBuffer;
 	}
 
@@ -910,6 +1000,7 @@ namespace AVEncoder
 			return;
 		}
 
+		PRAGMA_DISABLE_DEPRECATION_WARNINGS
 		switch (InBuffer->SourceFrame->GetCUDA().UnderlyingRHI)
 		{
 			case FVideoEncoderInputFrame::EUnderlyingRHI::Vulkan:
@@ -922,6 +1013,7 @@ namespace AVEncoder
 			default:
 				break;
 		}
+		PRAGMA_ENABLE_DEPRECATION_WARNINGS
 
 		InBuffer->Width = TextureSize.X;
 		InBuffer->Height = TextureSize.Y;
@@ -949,6 +1041,7 @@ namespace AVEncoder
 
 			void* TextureToCompress = nullptr;
 
+			PRAGMA_DISABLE_DEPRECATION_WARNINGS
 			switch (InBuffer->SourceFrame->GetFormat())
 			{
 #if PLATFORM_WINDOWS
@@ -966,6 +1059,7 @@ namespace AVEncoder
 				default:
 					break;
 			}
+			PRAGMA_ENABLE_DEPRECATION_WARNINGS
 
 			if (!TextureToCompress)
 			{
@@ -976,6 +1070,7 @@ namespace AVEncoder
 			InBuffer->InputTexture = TextureToCompress;
 			NVENCStruct(NV_ENC_REGISTER_RESOURCE, RegisterParam);
 
+			PRAGMA_DISABLE_DEPRECATION_WARNINGS
 			FIntPoint TextureSize(InBuffer->SourceFrame->GetWidth(), InBuffer->SourceFrame->GetHeight());
 
 			switch (InBuffer->SourceFrame->GetFormat())
@@ -993,6 +1088,7 @@ namespace AVEncoder
 				default:
 					break;
 			}
+			PRAGMA_ENABLE_DEPRECATION_WARNINGS
 
 			RegisterParam.resourceToRegister = InBuffer->InputTexture;
 
@@ -1307,7 +1403,9 @@ namespace AVEncoder
 		return OutSupportedProfiles != 0;
 	}
 
+	PRAGMA_DISABLE_DEPRECATION_WARNINGS
 	static bool GetEncoderSupportedInputFormats(FNVENCCommon& NVENC, void* InEncoder, TArray<EVideoFrameFormat>& OutSupportedInputFormats)
+	PRAGMA_ENABLE_DEPRECATION_WARNINGS
 	{
 		const uint32_t MaxInputFmtCount = 32;
 		uint32_t InputFmtCount = 0;
@@ -1330,18 +1428,22 @@ namespace AVEncoder
 				case NV_ENC_BUFFER_FORMAT_ARGB:
 					break;
 				case NV_ENC_BUFFER_FORMAT_ABGR:
+				PRAGMA_DISABLE_DEPRECATION_WARNINGS
 #if PLATFORM_WINDOWS
 					OutSupportedInputFormats.Push(EVideoFrameFormat::D3D11_R8G8B8A8_UNORM);
 					OutSupportedInputFormats.Push(EVideoFrameFormat::D3D12_R8G8B8A8_UNORM);
 #endif
 					OutSupportedInputFormats.Push(EVideoFrameFormat::CUDA_R8G8B8A8_UNORM);
+				PRAGMA_ENABLE_DEPRECATION_WARNINGS
 					break;
 			}
 		}
 		return true;
 	}
 
+	PRAGMA_DISABLE_DEPRECATION_WARNINGS
 	static bool GetEncoderInfo(FNVENCCommon& NVENC, FVideoEncoderInfo& EncoderInfo)
+	PRAGMA_ENABLE_DEPRECATION_WARNINGS
 	{
 		bool bSuccess = true;
 
@@ -1384,17 +1486,21 @@ namespace AVEncoder
 			return false;
 		}
 
+		PRAGMA_DISABLE_DEPRECATION_WARNINGS
 		EncoderInfo.CodecType = ECodecType::H264;
 		EncoderInfo.MaxWidth = GetEncoderCapability(NVENC, EncoderSession, NV_ENC_CAPS_WIDTH_MAX);
 		EncoderInfo.MaxHeight = GetEncoderCapability(NVENC, EncoderSession, NV_ENC_CAPS_HEIGHT_MAX);
+		PRAGMA_ENABLE_DEPRECATION_WARNINGS
 
 		int LevelMax = GetEncoderCapability(NVENC, EncoderSession, NV_ENC_CAPS_LEVEL_MAX);
 		int LevelMin = GetEncoderCapability(NVENC, EncoderSession, NV_ENC_CAPS_LEVEL_MIN);
 		if (LevelMin > 0 && LevelMax > 0 && LevelMax >= LevelMin)
 		{
+			PRAGMA_DISABLE_DEPRECATION_WARNINGS
 			EncoderInfo.H264.MinLevel = (LevelMin > 9) ? LevelMin : 9;
 			EncoderInfo.H264.MaxLevel = (LevelMax < 9) ? 9 : (LevelMax > NV_ENC_LEVEL_H264_52) ? NV_ENC_LEVEL_H264_52
 																							   : LevelMax;
+			PRAGMA_ENABLE_DEPRECATION_WARNINGS
 		}
 		else
 		{
@@ -1402,7 +1508,9 @@ namespace AVEncoder
 			bSuccess = false;
 		}
 
+		PRAGMA_DISABLE_DEPRECATION_WARNINGS
 		if (!GetEncoderSupportedProfiles(NVENC, EncoderSession, EncoderInfo.H264.SupportedProfiles) || !GetEncoderSupportedInputFormats(NVENC, EncoderSession, EncoderInfo.SupportedInputFormats))
+		PRAGMA_ENABLE_DEPRECATION_WARNINGS
 		{
 			bSuccess = false;
 		}
diff --git a/Engine/Plugins/Media/HardwareEncoders/Source/EncoderNVENC/Private/NVENC_EncoderH264.h b/Engine/Plugins/Media/HardwareEncoders/Source/EncoderNVENC/Private/NVENC_EncoderH264.h
--- a/Engine/Plugins/Media/HardwareEncoders/Source/EncoderNVENC/Private/NVENC_EncoderH264.h
+++ b/Engine/Plugins/Media/HardwareEncoders/Source/EncoderNVENC/Private/NVENC_EncoderH264.h
@@ -14,32 +14,44 @@ namespace AVEncoder { class FVideoEncoderInputFrameImpl; }
 
 namespace AVEncoder
 {
+	PRAGMA_DISABLE_DEPRECATION_WARNINGS
 	class FVideoEncoderNVENC_H264 : public FVideoEncoder
+	PRAGMA_ENABLE_DEPRECATION_WARNINGS
 	{
 
 	public:
 		virtual ~FVideoEncoderNVENC_H264() override; 
 
+		PRAGMA_DISABLE_DEPRECATION_WARNINGS
 		bool Setup(TSharedRef<FVideoEncoderInput> InputFrameFactory, const FLayerConfig& InitConfig) override;
+		PRAGMA_ENABLE_DEPRECATION_WARNINGS
 		void Shutdown() override;
 
 		// Query whether or not encoder is supported and available
+		PRAGMA_DISABLE_DEPRECATION_WARNINGS
 		static bool GetIsAvailable(const FVideoEncoderInput& InFrameFactory, FVideoEncoderInfo& OutEncoderInfo);
+		PRAGMA_ENABLE_DEPRECATION_WARNINGS
 
 		// Register encoder with video encoder factory
 		static void Register(FVideoEncoderFactory& InFactory);
 
+		PRAGMA_DISABLE_DEPRECATION_WARNINGS
 		void Encode(const TSharedPtr<FVideoEncoderInputFrame> InFrame, const FEncodeOptions& EncodeOptions) override;
+		PRAGMA_ENABLE_DEPRECATION_WARNINGS
 		void Flush();
 
 	protected:
+		PRAGMA_DISABLE_DEPRECATION_WARNINGS
 		FLayer* CreateLayer(uint32 InLayerIndex, const FLayerConfig& InLayerConfig) override;
 		void DestroyLayer(FLayer* layer) override;
+		PRAGMA_ENABLE_DEPRECATION_WARNINGS
 
 	private:
 		FVideoEncoderNVENC_H264();
 
+		PRAGMA_DISABLE_DEPRECATION_WARNINGS
 		class FNVENCLayer : public FLayer
+		PRAGMA_ENABLE_DEPRECATION_WARNINGS
 		{
 
 		public:
@@ -63,7 +75,9 @@ namespace AVEncoder
 				uint64 SubmitTimeCycles;
 			};
 
+			PRAGMA_DISABLE_DEPRECATION_WARNINGS
 			FNVENCLayer(uint32 layerIdx, FLayerConfig const& config, FVideoEncoderNVENC_H264& encoder);
+			PRAGMA_ENABLE_DEPRECATION_WARNINGS
 			~FNVENCLayer();
 			bool Setup();
 			bool CreateSession();
@@ -73,7 +87,9 @@ namespace AVEncoder
 			void MaybeReconfigure();
 			void UpdateConfig();
 			void UpdateLastEncodedQP(uint32 InLastEncodedQP);
+			PRAGMA_DISABLE_DEPRECATION_WARNINGS
 			void Encode(const TSharedPtr<FVideoEncoderInputFrame> InFrame, const FEncodeOptions& EncodeOptions);
+			PRAGMA_ENABLE_DEPRECATION_WARNINGS
 			void EncodeBuffer(FInputOutput* Buffer);
 			void ProcessEncodedBuffer(FInputOutput* Buffer);
 			void Flush();
@@ -110,7 +126,9 @@ namespace AVEncoder
 		};
 
 		FNVENCCommon& NVENC;
+		PRAGMA_DISABLE_DEPRECATION_WARNINGS
 		EVideoFrameFormat FrameFormat = EVideoFrameFormat::Undefined;
+		PRAGMA_ENABLE_DEPRECATION_WARNINGS
 
 		// this could be a TRefCountPtr<ID3D11Device>, CUcontext or void*
 		void* EncoderDevice;
diff --git a/Engine/Plugins/Media/HardwareEncoders/Source/EncoderNVENC/Private/NVENC_EncoderModule.cpp b/Engine/Plugins/Media/HardwareEncoders/Source/EncoderNVENC/Private/NVENC_EncoderModule.cpp
--- a/Engine/Plugins/Media/HardwareEncoders/Source/EncoderNVENC/Private/NVENC_EncoderModule.cpp
+++ b/Engine/Plugins/Media/HardwareEncoders/Source/EncoderNVENC/Private/NVENC_EncoderModule.cpp
@@ -43,8 +43,11 @@ public:
 				{
 					if (IsRHIDeviceNVIDIA())
 					{
+						PRAGMA_DISABLE_DEPRECATION_WARNINGS
 						FVideoEncoderNVENC_H264::Register(FVideoEncoderFactory::Get());
+						PRAGMA_ENABLE_DEPRECATION_WARNINGS
 					}
+					
 				});
 			}
 		}
diff --git a/Engine/Source/Runtime/AVEncoder/Private/AudioEncoder.cpp b/Engine/Source/Runtime/AVEncoder/Private/AudioEncoder.cpp
--- a/Engine/Source/Runtime/AVEncoder/Private/AudioEncoder.cpp
+++ b/Engine/Source/Runtime/AVEncoder/Private/AudioEncoder.cpp
@@ -5,26 +5,35 @@
 
 namespace AVEncoder
 {
+	PRAGMA_DISABLE_DEPRECATION_WARNINGS
 	void FAudioEncoder::RegisterListener(IAudioEncoderListener& Listener)
+	PRAGMA_ENABLE_DEPRECATION_WARNINGS
 	{
 		FScopeLock lock{ &ListenersMutex };
 		check(Listeners.Find(&Listener) == INDEX_NONE);
 		Listeners.AddUnique(&Listener);
 	}
 
+	PRAGMA_DISABLE_DEPRECATION_WARNINGS
 	void FAudioEncoder::UnregisterListener(IAudioEncoderListener& Listener)
+	PRAGMA_ENABLE_DEPRECATION_WARNINGS
 	{
 		FScopeLock lock{ &ListenersMutex };
 		int32 Count = Listeners.Remove(&Listener);
 		check(Count == 1);
 	}
 
+	PRAGMA_DISABLE_DEPRECATION_WARNINGS
 	void FAudioEncoder::OnEncodedAudioFrame(const FMediaPacket& Packet)
+	PRAGMA_ENABLE_DEPRECATION_WARNINGS
 	{
 		FScopeLock lock{ &ListenersMutex };
 		for (auto&& L : Listeners)
 		{
+			PRAGMA_DISABLE_DEPRECATION_WARNINGS
 			L->OnEncodedAudioFrame(Packet);
+			PRAGMA_ENABLE_DEPRECATION_WARNINGS
 		}
+		
 	}
 }
diff --git a/Engine/Source/Runtime/AVEncoder/Private/AudioEncoderFactory.cpp b/Engine/Source/Runtime/AVEncoder/Private/AudioEncoderFactory.cpp
--- a/Engine/Source/Runtime/AVEncoder/Private/AudioEncoderFactory.cpp
+++ b/Engine/Source/Runtime/AVEncoder/Private/AudioEncoderFactory.cpp
@@ -19,21 +19,29 @@ namespace AVEncoder
 		RegisterDefaultFactories();
 	}
 
+	PRAGMA_DISABLE_DEPRECATION_WARNINGS
 	TArray<FAudioEncoderFactory*> FAudioEncoderFactory::Factories;
+	PRAGMA_ENABLE_DEPRECATION_WARNINGS
 
+	PRAGMA_DISABLE_DEPRECATION_WARNINGS
 	void FAudioEncoderFactory::RegisterFactory(FAudioEncoderFactory& Factory)
+	PRAGMA_ENABLE_DEPRECATION_WARNINGS
 	{
 		DoDefaultRegistration();
 
 		Factories.AddUnique(&Factory);
 	}
 
+	PRAGMA_DISABLE_DEPRECATION_WARNINGS
 	void FAudioEncoderFactory::UnregisterFactory(FAudioEncoderFactory& Factory)
+	PRAGMA_ENABLE_DEPRECATION_WARNINGS
 	{
 		Factories.Remove(&Factory);
 	}
 
+	PRAGMA_DISABLE_DEPRECATION_WARNINGS
 	FAudioEncoderFactory* FAudioEncoderFactory::FindFactory(const FString& Codec)
+	PRAGMA_ENABLE_DEPRECATION_WARNINGS
 	{
 		DoDefaultRegistration();
 
@@ -48,7 +56,9 @@ namespace AVEncoder
 		return nullptr;
 	}
 
+	PRAGMA_DISABLE_DEPRECATION_WARNINGS
 	const TArray<FAudioEncoderFactory*> FAudioEncoderFactory::GetAllFactories()
+	PRAGMA_ENABLE_DEPRECATION_WARNINGS
 	{
 		DoDefaultRegistration();
 
@@ -63,9 +73,12 @@ namespace AVEncoder
 #if PLATFORM_WINDOWS
 		// Generic Windows/XBox Wmf encoder
 		static FWmfAudioEncoderFactory WmfAudioEncoderFactory;
+		PRAGMA_DISABLE_DEPRECATION_WARNINGS
 		FAudioEncoderFactory::RegisterFactory(WmfAudioEncoderFactory);
+		PRAGMA_ENABLE_DEPRECATION_WARNINGS
 #endif
 
+
 		// Log all available encoders
 		{
 			auto CodecsInfo = [&](auto&& Factories) -> FString
@@ -73,8 +86,11 @@ namespace AVEncoder
 				FString Str;
 				for (auto&& Factory : Factories)
 				{
+					PRAGMA_DISABLE_DEPRECATION_WARNINGS
 					Str += FString::Printf(TEXT(", %s(%s) "), Factory->GetName(), *FString::Join(Factory->GetSupportedCodecs(), TEXT("/")));
+					PRAGMA_ENABLE_DEPRECATION_WARNINGS
 				}
+				
 				if (Str.IsEmpty())
 				{
 					return FString(TEXT("None"));
@@ -85,7 +101,10 @@ namespace AVEncoder
 				}
 			};
 
+			PRAGMA_DISABLE_DEPRECATION_WARNINGS
 			UE_LOG(LogAVEncoder, Log, TEXT("Available audio encoders: %s "), *CodecsInfo(FAudioEncoderFactory::GetAllFactories()));
+			PRAGMA_ENABLE_DEPRECATION_WARNINGS
 		}
+		
 	}
 }
\ No newline at end of file
diff --git a/Engine/Source/Runtime/AVEncoder/Private/CodecPacket.cpp b/Engine/Source/Runtime/AVEncoder/Private/CodecPacket.cpp
--- a/Engine/Source/Runtime/AVEncoder/Private/CodecPacket.cpp
+++ b/Engine/Source/Runtime/AVEncoder/Private/CodecPacket.cpp
@@ -21,7 +21,9 @@ namespace
 
 namespace AVEncoder
 {
+	PRAGMA_DISABLE_DEPRECATION_WARNINGS
 	FCodecPacket FCodecPacket::Create(const uint8* InData, uint32 InDataSize)
+	PRAGMA_ENABLE_DEPRECATION_WARNINGS
 	{
 		uint8* DataCopy = static_cast<uint8*>(FMemory::Malloc(InDataSize));
 		FMemory::BigBlockMemcpy(DataCopy, InData, InDataSize);
diff --git a/Engine/Source/Runtime/AVEncoder/Private/Decoders/VideoDecoder.cpp b/Engine/Source/Runtime/AVEncoder/Private/Decoders/VideoDecoder.cpp
--- a/Engine/Source/Runtime/AVEncoder/Private/Decoders/VideoDecoder.cpp
+++ b/Engine/Source/Runtime/AVEncoder/Private/Decoders/VideoDecoder.cpp
@@ -11,7 +11,9 @@ namespace AVEncoder
 
 struct FVideoDecoder::FPlatformDecoderAllocInterface
 {
+	PRAGMA_DISABLE_DEPRECATION_WARNINGS
 	FVideoDecoderMethodsWindows	AllocMethods = {};
+	PRAGMA_ENABLE_DEPRECATION_WARNINGS
 	void*						ClientData = nullptr;
 };
 
@@ -27,13 +29,15 @@ bool FVideoDecoder::CreateDecoderAllocationInterface()
 	if (!PlatformDecoderAllocInterface)
 	{
 		check(PlatformDecoderAllocInterface == nullptr);
-		FPlatformDecoderAllocInterface* Interface = new FPlatformDecoderAllocInterface {};
 
-// TODO: Platform specific!!
+		PRAGMA_DISABLE_DEPRECATION_WARNINGS
+		FPlatformDecoderAllocInterface* Interface = new FPlatformDecoderAllocInterface {};
+		// TODO: Platform specific!!
 		Interface->AllocMethods.MagicCookie = 0x57696e58;		// 'WinX'
 		Interface->AllocMethods.This = nullptr;					// This will become the client's 'this' pointer
 		Interface->AllocMethods.GetD3DDevice = nullptr;			// Method to get the current D3D device from the client.
 		Interface->AllocMethods.AllocateFrameBuffer = nullptr;	// Method to allocate frame buffer memory from the client.
+		PRAGMA_ENABLE_DEPRECATION_WARNINGS
 
 		// Call into the client to set up the allocation interface.
 		CreateDecoderAllocationInterfaceFN(&Interface->AllocMethods, &Interface->ClientData);
@@ -53,6 +57,7 @@ bool FVideoDecoder::CreateDecoderAllocationInterface()
 	return true;
 }
 
+PRAGMA_DISABLE_DEPRECATION_WARNINGS
 void FVideoDecoder::ReleaseDecoderAllocationInterface()
 {
 	if (PlatformDecoderAllocInterface)
@@ -62,24 +67,30 @@ void FVideoDecoder::ReleaseDecoderAllocationInterface()
 		delete PlatformDecoderAllocInterface;
 		PlatformDecoderAllocInterface = nullptr;
 	}
-
 }
+PRAGMA_ENABLE_DEPRECATION_WARNINGS
 
+PRAGMA_DISABLE_DEPRECATION_WARNINGS
 void* FVideoDecoder::GetAllocationInterfaceMethods()
 {
 	return PlatformDecoderAllocInterface ? &PlatformDecoderAllocInterface->AllocMethods : nullptr;
+PRAGMA_ENABLE_DEPRECATION_WARNINGS
 }
 
 
+PRAGMA_DISABLE_DEPRECATION_WARNINGS
 EFrameBufferAllocReturn FVideoDecoder::AllocateOutputFrameBuffer(FVideoDecoderAllocFrameBufferResult* OutBuffer, const FVideoDecoderAllocFrameBufferParams* InAllocParams)
+PRAGMA_ENABLE_DEPRECATION_WARNINGS
 {
 	check(InAllocParams);
 	check(OutBuffer);
+	PRAGMA_DISABLE_DEPRECATION_WARNINGS
 	if (PlatformDecoderAllocInterface && PlatformDecoderAllocInterface->AllocMethods.AllocateFrameBuffer.IsBound())
 	{
 		return PlatformDecoderAllocInterface->AllocMethods.AllocateFrameBuffer.Execute(PlatformDecoderAllocInterface->AllocMethods.This, InAllocParams, OutBuffer);
 	}
 	return EFrameBufferAllocReturn::CODEC_Failure;
+	PRAGMA_ENABLE_DEPRECATION_WARNINGS
 }
 
 
diff --git a/Engine/Source/Runtime/AVEncoder/Private/Decoders/VideoDecoderH264_Dummy.cpp b/Engine/Source/Runtime/AVEncoder/Private/Decoders/VideoDecoderH264_Dummy.cpp
--- a/Engine/Source/Runtime/AVEncoder/Private/Decoders/VideoDecoderH264_Dummy.cpp
+++ b/Engine/Source/Runtime/AVEncoder/Private/Decoders/VideoDecoderH264_Dummy.cpp
@@ -8,16 +8,22 @@
 namespace AVEncoder
 {
 
+PRAGMA_DISABLE_DEPRECATION_WARNINGS
 class FVideoDecoderOutputDummy : public FVideoDecoderOutput
+PRAGMA_ENABLE_DEPRECATION_WARNINGS
 {
 public:
+	PRAGMA_DISABLE_DEPRECATION_WARNINGS
 	FVideoDecoderOutputDummy(int32 w, int32 h, int64 pts)
 	: Width(w), Height(h), PTS(pts)
 	{}
+	PRAGMA_ENABLE_DEPRECATION_WARNINGS
 
+	PRAGMA_DISABLE_DEPRECATION_WARNINGS
 	virtual ~FVideoDecoderOutputDummy()
 	{
 	}
+	PRAGMA_ENABLE_DEPRECATION_WARNINGS
 
 	virtual int32 AddRef() override
 	{
@@ -50,7 +56,9 @@ public:
 		return PTS;
 	}
 
+	PRAGMA_DISABLE_DEPRECATION_WARNINGS
 	virtual const FVideoDecoderAllocFrameBufferResult* GetAllocatedBuffer() const
+	PRAGMA_ENABLE_DEPRECATION_WARNINGS
 	{
 		return &Buffer;
 	}
@@ -93,13 +101,17 @@ public:
 	}
 
 	// Internal for allocation.
+	PRAGMA_DISABLE_DEPRECATION_WARNINGS
 	FVideoDecoderAllocFrameBufferResult* GetBuffer()
+	PRAGMA_ENABLE_DEPRECATION_WARNINGS
 	{
 		return &Buffer;
 	}
 
 private:
+	PRAGMA_DISABLE_DEPRECATION_WARNINGS
 	FVideoDecoderAllocFrameBufferResult	Buffer = {};
+	PRAGMA_ENABLE_DEPRECATION_WARNINGS
 	int32	RefCount = 1;
 	int32	Width = 0;
 	int32	Height = 0;
@@ -109,30 +121,44 @@ private:
 	// query whether or not encoder is supported and available
 //	static bool GetIsAvailable(FVideoEncoderInputImpl& InInput, FVideoEncoderInfo& OutEncoderInfo);
 
+PRAGMA_DISABLE_DEPRECATION_WARNINGS
 void FVideoDecoderH264_Dummy::Register(FVideoDecoderFactory& InFactory)
+PRAGMA_ENABLE_DEPRECATION_WARNINGS
 {
+	PRAGMA_DISABLE_DEPRECATION_WARNINGS
 	FVideoDecoderInfo	DecoderInfo;
 	DecoderInfo.CodecType = ECodecType::H264;
 	DecoderInfo.MaxWidth = 1920;
 	DecoderInfo.MaxHeight = 1088;
+	PRAGMA_ENABLE_DEPRECATION_WARNINGS
 
+	PRAGMA_DISABLE_DEPRECATION_WARNINGS
 	InFactory.Register(DecoderInfo, []() {
-			return new FVideoDecoderH264_Dummy();
-		});
+		return new FVideoDecoderH264_Dummy();
+	});
+	PRAGMA_ENABLE_DEPRECATION_WARNINGS
 }
 
+PRAGMA_DISABLE_DEPRECATION_WARNINGS
 FVideoDecoderH264_Dummy::FVideoDecoderH264_Dummy()
 {
 }
+PRAGMA_ENABLE_DEPRECATION_WARNINGS
 
+PRAGMA_DISABLE_DEPRECATION_WARNINGS
 FVideoDecoderH264_Dummy::~FVideoDecoderH264_Dummy()
 {
 }
+PRAGMA_ENABLE_DEPRECATION_WARNINGS
 
+PRAGMA_DISABLE_DEPRECATION_WARNINGS
 bool FVideoDecoderH264_Dummy::Setup(const FVideoDecoder::FInit& InInit)
+PRAGMA_ENABLE_DEPRECATION_WARNINGS
 {
+	PRAGMA_DISABLE_DEPRECATION_WARNINGS
 	CreateDecoderAllocationInterfaceFN = InInit.CreateDecoderAllocationInterface;
 	ReleaseDecoderAllocationInterfaceFN = InInit.ReleaseDecoderAllocationInterface;
+	PRAGMA_ENABLE_DEPRECATION_WARNINGS
 	return true;
 }
 
@@ -140,31 +166,42 @@ void FVideoDecoderH264_Dummy::Shutdown()
 {
 	if (bIsInitialized)
 	{
+		PRAGMA_DISABLE_DEPRECATION_WARNINGS
 		ReleaseDecoderAllocationInterface();
+		PRAGMA_ENABLE_DEPRECATION_WARNINGS
 		bIsInitialized = false;
 	}
 	delete this;
 }
 
 
+PRAGMA_DISABLE_DEPRECATION_WARNINGS
 FVideoDecoder::EDecodeResult FVideoDecoderH264_Dummy::Decode(const FVideoDecoderInput* InInput)
+PRAGMA_ENABLE_DEPRECATION_WARNINGS
 {
 	// Lazily initialize the decoder.
 	if (!bIsInitialized)
 	{
+		PRAGMA_DISABLE_DEPRECATION_WARNINGS
 		bIsInitialized = CreateDecoderAllocationInterface();
+		PRAGMA_ENABLE_DEPRECATION_WARNINGS
 	}
+	
 
 	//
+	PRAGMA_DISABLE_DEPRECATION_WARNINGS
 	if (OnDecodedFrame)
+	PRAGMA_ENABLE_DEPRECATION_WARNINGS
 	{
 		if (bIsInitialized)
 		{
+			PRAGMA_DISABLE_DEPRECATION_WARNINGS
 			int32_t	Width = InInput->GetWidth();
 			int32_t Height = InInput->GetHeight();
-
+			
 			TUniquePtr<FVideoDecoderOutputDummy> pNew(new FVideoDecoderOutputDummy(Width, Height, InInput->GetPTS()));
 
+
 			// Get memory from the application
 			FVideoDecoderAllocFrameBufferParams ap {};
 			EFrameBufferAllocReturn ar;
@@ -192,10 +229,13 @@ FVideoDecoder::EDecodeResult FVideoDecoderH264_Dummy::Decode(const FVideoDecoder
 				// Error!
 				return FVideoDecoder::EDecodeResult::Failure;
 			}
+			PRAGMA_ENABLE_DEPRECATION_WARNINGS
 
 			// Buffer needs to contain an NV12 texture.
 			int32 ySize = Width * Height;
+			PRAGMA_DISABLE_DEPRECATION_WARNINGS
 			uint8* pY = (uint8*)pNew->GetBuffer()->AllocatedBuffer;
+			PRAGMA_ENABLE_DEPRECATION_WARNINGS
 			uint8* pUV = pY + ySize;
 			FMemory::Memset(pUV, 128, ySize / 2);
 			int C = ++YOffset;
@@ -204,19 +244,25 @@ FVideoDecoder::EDecodeResult FVideoDecoderH264_Dummy::Decode(const FVideoDecoder
 				FMemory::Memset(pY, (uint8)C++, Width);
 				pY += Width;
 			}
+			PRAGMA_DISABLE_DEPRECATION_WARNINGS
 			OnDecodedFrame(pNew.Release());
 			return FVideoDecoder::EDecodeResult::Success;
+			PRAGMA_ENABLE_DEPRECATION_WARNINGS
 		}
 		else
 		{
 // TODO: could not decode, but may need to still notify the caller with an empty frame??
+			PRAGMA_DISABLE_DEPRECATION_WARNINGS
 			return FVideoDecoder::EDecodeResult::Failure;
+			PRAGMA_ENABLE_DEPRECATION_WARNINGS
 		}
 	}
 	else
 	{
 		// With no registered callback that's interested in the result we can presume we would have been successful.
+		PRAGMA_DISABLE_DEPRECATION_WARNINGS
 		return FVideoDecoder::EDecodeResult::Success;
+		PRAGMA_ENABLE_DEPRECATION_WARNINGS
 	}
 }
 
diff --git a/Engine/Source/Runtime/AVEncoder/Private/Decoders/VideoDecoderH264_Dummy.h b/Engine/Source/Runtime/AVEncoder/Private/Decoders/VideoDecoderH264_Dummy.h
--- a/Engine/Source/Runtime/AVEncoder/Private/Decoders/VideoDecoderH264_Dummy.h
+++ b/Engine/Source/Runtime/AVEncoder/Private/Decoders/VideoDecoderH264_Dummy.h
@@ -12,16 +12,20 @@
 namespace AVEncoder
 {
 
+PRAGMA_DISABLE_DEPRECATION_WARNINGS
 class FVideoDecoderH264_Dummy : public FVideoDecoder
+PRAGMA_ENABLE_DEPRECATION_WARNINGS
 {
 public:
 	// register decoder with video decoder factory
+	PRAGMA_DISABLE_DEPRECATION_WARNINGS
 	static void Register(FVideoDecoderFactory& InFactory);
-
 	virtual bool Setup(const FInit& InInit) override;
+	virtual EDecodeResult Decode(const FVideoDecoderInput* InInput) override;
+	PRAGMA_ENABLE_DEPRECATION_WARNINGS
+
 	virtual void Shutdown() override;
 
-	virtual EDecodeResult Decode(const FVideoDecoderInput* InInput) override;
 
 
 	// query whether or not encoder is supported and available
diff --git a/Engine/Source/Runtime/AVEncoder/Private/Decoders/VideoDecoderH264_Impl.cpp b/Engine/Source/Runtime/AVEncoder/Private/Decoders/VideoDecoderH264_Impl.cpp
--- a/Engine/Source/Runtime/AVEncoder/Private/Decoders/VideoDecoderH264_Impl.cpp
+++ b/Engine/Source/Runtime/AVEncoder/Private/Decoders/VideoDecoderH264_Impl.cpp
@@ -21,7 +21,9 @@ namespace AVEncoder
 #if CODEC_HAVE_NATIVE_H264_DECODER
 
 #if PLATFORM_WINDOWS
+PRAGMA_DISABLE_DEPRECATION_WARNINGS
 void FVideoDecoderH264_Impl::Register(FVideoDecoderFactory& InFactory)
+PRAGMA_ENABLE_DEPRECATION_WARNINGS
 {
 	FVideoDecoderH264_Windows::Register(InFactory);
 }
@@ -29,7 +31,9 @@ void FVideoDecoderH264_Impl::Register(FVideoDecoderFactory& InFactory)
 
 #else
 
+PRAGMA_DISABLE_DEPRECATION_WARNINGS
 void FVideoDecoderH264_Impl::Register(FVideoDecoderFactory& InFactory)
+PRAGMA_ENABLE_DEPRECATION_WARNINGS
 {
 }
 
diff --git a/Engine/Source/Runtime/AVEncoder/Private/Decoders/VideoDecoderH264_Impl.h b/Engine/Source/Runtime/AVEncoder/Private/Decoders/VideoDecoderH264_Impl.h
--- a/Engine/Source/Runtime/AVEncoder/Private/Decoders/VideoDecoderH264_Impl.h
+++ b/Engine/Source/Runtime/AVEncoder/Private/Decoders/VideoDecoderH264_Impl.h
@@ -8,18 +8,25 @@
 namespace AVEncoder
 {
 
+PRAGMA_DISABLE_DEPRECATION_WARNINGS
 class FVideoDecoderH264_Impl : public FVideoDecoder
+PRAGMA_ENABLE_DEPRECATION_WARNINGS
 {
 public:
 	// register decoder with video decoder factory
+	PRAGMA_DISABLE_DEPRECATION_WARNINGS
 	static void Register(FVideoDecoderFactory& InFactory);
-
 	virtual bool Setup(const FInit& InInit) = 0;
+	virtual EDecodeResult Decode(const FVideoDecoderInput* InInput) = 0;
+	PRAGMA_ENABLE_DEPRECATION_WARNINGS
+
 	virtual void Shutdown() = 0;
 
-	virtual EDecodeResult Decode(const FVideoDecoderInput* InInput) = 0;
 protected:
+	PRAGMA_DISABLE_DEPRECATION_WARNINGS
 	virtual ~FVideoDecoderH264_Impl() = default;
+	PRAGMA_ENABLE_DEPRECATION_WARNINGS
 };
 
+
 } // namespace AVEncoder
diff --git a/Engine/Source/Runtime/AVEncoder/Private/Decoders/VideoDecoderMpeg4.cpp b/Engine/Source/Runtime/AVEncoder/Private/Decoders/VideoDecoderMpeg4.cpp
--- a/Engine/Source/Runtime/AVEncoder/Private/Decoders/VideoDecoderMpeg4.cpp
+++ b/Engine/Source/Runtime/AVEncoder/Private/Decoders/VideoDecoderMpeg4.cpp
@@ -17,16 +17,22 @@
 namespace AVEncoder
 {
 
+PRAGMA_DISABLE_DEPRECATION_WARNINGS
 class FVideoDecoderOutputMPEG4 : public FVideoDecoderOutput
+PRAGMA_ENABLE_DEPRECATION_WARNINGS
 {
 public:
+	PRAGMA_DISABLE_DEPRECATION_WARNINGS
 	FVideoDecoderOutputMPEG4(int32 w, int32 h, int64 pts)
 	: Width(w), Height(h), PTS(pts)
 	{}
+	PRAGMA_ENABLE_DEPRECATION_WARNINGS
 
+	PRAGMA_DISABLE_DEPRECATION_WARNINGS
 	virtual ~FVideoDecoderOutputMPEG4()
 	{
 	}
+	PRAGMA_ENABLE_DEPRECATION_WARNINGS
 
 	virtual int32 AddRef() override
 	{
@@ -59,7 +65,9 @@ public:
 		return PTS;
 	}
 
+	PRAGMA_DISABLE_DEPRECATION_WARNINGS
 	virtual const FVideoDecoderAllocFrameBufferResult* GetAllocatedBuffer() const
+	PRAGMA_ENABLE_DEPRECATION_WARNINGS
 	{
 		return &Buffer;
 	}
@@ -106,13 +114,17 @@ public:
 	}
 
 	// Internal for allocation.
+	PRAGMA_DISABLE_DEPRECATION_WARNINGS
 	FVideoDecoderAllocFrameBufferResult* GetBuffer()
+	PRAGMA_ENABLE_DEPRECATION_WARNINGS
 	{
 		return &Buffer;
 	}
 
 private:
+	PRAGMA_DISABLE_DEPRECATION_WARNINGS
 	FVideoDecoderAllocFrameBufferResult	Buffer = {};
+	PRAGMA_ENABLE_DEPRECATION_WARNINGS
 	int32	RefCount = 1;
 	int32	Width = 0;
 	int32	Pitch = 0;
@@ -135,10 +147,14 @@ static TMap<void*, uint32>& Actives()
 class FVideoDecoderMPEG4_Impl : public FVideoDecoderMPEG4, public vdecmpeg4::VIDStreamIO, public vdecmpeg4::VIDStreamEvents
 {
 public:
+	PRAGMA_DISABLE_DEPRECATION_WARNINGS
 	virtual bool Setup(const FInit& InInit) override;
+	PRAGMA_ENABLE_DEPRECATION_WARNINGS
 	virtual void Shutdown() override;
 
+	PRAGMA_DISABLE_DEPRECATION_WARNINGS
 	virtual EDecodeResult Decode(const FVideoDecoderInput* InInput) override;
+	PRAGMA_ENABLE_DEPRECATION_WARNINGS
 
 	FVideoDecoderMPEG4_Impl();
 	virtual ~FVideoDecoderMPEG4_Impl();
@@ -196,16 +212,22 @@ protected:
 
 
 
+PRAGMA_DISABLE_DEPRECATION_WARNINGS
 void FVideoDecoderMPEG4::Register(FVideoDecoderFactory& InFactory)
+PRAGMA_ENABLE_DEPRECATION_WARNINGS
 {
+	PRAGMA_DISABLE_DEPRECATION_WARNINGS
 	FVideoDecoderInfo	DecoderInfo;
 	DecoderInfo.CodecType = ECodecType::MPEG4;
 	DecoderInfo.MaxWidth = 1920;
 	DecoderInfo.MaxHeight = 1088;
+	PRAGMA_ENABLE_DEPRECATION_WARNINGS
 
+	PRAGMA_DISABLE_DEPRECATION_WARNINGS
 	InFactory.Register(DecoderInfo, []() {
-			return new FVideoDecoderMPEG4_Impl();
-		});
+		return new FVideoDecoderMPEG4_Impl();
+	});
+	PRAGMA_ENABLE_DEPRECATION_WARNINGS
 }
 
 int32 gTestStreamIndex4 = 0;
@@ -224,10 +246,14 @@ FVideoDecoderMPEG4_Impl::~FVideoDecoderMPEG4_Impl()
 	check(!DecoderHandle);
 }
 
+PRAGMA_DISABLE_DEPRECATION_WARNINGS
 bool FVideoDecoderMPEG4_Impl::Setup(const FVideoDecoder::FInit& InInit)
+PRAGMA_ENABLE_DEPRECATION_WARNINGS
 {
+	PRAGMA_DISABLE_DEPRECATION_WARNINGS
 	CreateDecoderAllocationInterfaceFN = InInit.CreateDecoderAllocationInterface;
 	ReleaseDecoderAllocationInterfaceFN = InInit.ReleaseDecoderAllocationInterface;
+	PRAGMA_ENABLE_DEPRECATION_WARNINGS
 	return true;
 }
 
@@ -240,7 +266,9 @@ void FVideoDecoderMPEG4_Impl::Shutdown()
 			vdecmpeg4::VIDDestroyDecoder(DecoderHandle);
 			DecoderHandle = nullptr;
 		}
+		PRAGMA_DISABLE_DEPRECATION_WARNINGS
 		ReleaseDecoderAllocationInterface();
+		PRAGMA_ENABLE_DEPRECATION_WARNINGS
 		bIsInitialized = false;
 	}
 	delete this;
@@ -250,7 +278,9 @@ bool FVideoDecoderMPEG4_Impl::FirstUseInit()
 {
 	if (!bIsInitialized)
 	{
+		PRAGMA_DISABLE_DEPRECATION_WARNINGS
 		bIsInitialized = CreateDecoderAllocationInterface();
+		PRAGMA_ENABLE_DEPRECATION_WARNINGS
 		if (bIsInitialized)
 		{
 			FMemory::Memzero(DecoderSetup);
@@ -325,7 +355,9 @@ bool FVideoDecoderMPEG4_Impl::IsEof()
 }
 
 
+PRAGMA_DISABLE_DEPRECATION_WARNINGS
 static void CopyI420ToNV12(const FVideoDecoderAllocFrameBufferResult* OutBuf, const vdecmpeg4::VIDImage* vid)
+PRAGMA_ENABLE_DEPRECATION_WARNINGS
 {
 	int32 Width = vid->width;
 	int32 Height = vid->height;
@@ -335,6 +367,7 @@ static void CopyI420ToNV12(const FVideoDecoderAllocFrameBufferResult* OutBuf, co
 	const uint8_t* srcV = vid->v;
 
 	// Allocated buffer needs to have 3 planes.
+	PRAGMA_DISABLE_DEPRECATION_WARNINGS
 	check(OutBuf->AllocatedPlanesNum == 3);
 
 	uint8* OutBufferBase = (uint8*)OutBuf->AllocatedBuffer;
@@ -342,22 +375,28 @@ static void CopyI420ToNV12(const FVideoDecoderAllocFrameBufferResult* OutBuf, co
 	check(OutBuf->AllocatedPlaneDesc[0].BytesPerPixel == 1);
 	check(OutBuf->AllocatedPlaneDesc[0].ByteOffsetBetweenPixels == 1);
 	uint8* dstY = OutBufferBase + OutBuf->AllocatedPlaneDesc[0].ByteOffsetToFirstPixel;
+	PRAGMA_ENABLE_DEPRECATION_WARNINGS
 	// Need to copy the Y plane row by row.
 	for(int32 y=0; y<Height; ++y)
 	{
 		FMemory::Memcpy(dstY, srcY, Width);
 		srcY += vid->texWidth;
+		PRAGMA_DISABLE_DEPRECATION_WARNINGS
 		dstY += OutBuf->AllocatedPlaneDesc[0].ByteOffsetBetweenRows;
+		PRAGMA_ENABLE_DEPRECATION_WARNINGS
 	}
+	
 
 	// The U and V plane must be interleaved for NV12. We don't specifically do interleaving here but
 	// instead rely on the output plane description to be set up accordingly.
+	PRAGMA_DISABLE_DEPRECATION_WARNINGS
 	check(OutBuf->AllocatedPlaneDesc[1].BytesPerPixel == 1);
 	check(OutBuf->AllocatedPlaneDesc[2].BytesPerPixel == 1);
 	uint8* dstU = OutBufferBase + OutBuf->AllocatedPlaneDesc[1].ByteOffsetToFirstPixel;
 	uint8* dstV = OutBufferBase + OutBuf->AllocatedPlaneDesc[2].ByteOffsetToFirstPixel;
 	const int32 uOffCol = OutBuf->AllocatedPlaneDesc[1].ByteOffsetBetweenPixels;
 	const int32 vOffCol = OutBuf->AllocatedPlaneDesc[2].ByteOffsetBetweenPixels;
+	PRAGMA_ENABLE_DEPRECATION_WARNINGS
 	for(int32 v=0, vMax=Height/2; v<vMax; ++v)
 	{
 		uint8* U = dstU;
@@ -371,27 +410,36 @@ static void CopyI420ToNV12(const FVideoDecoderAllocFrameBufferResult* OutBuf, co
 		}
 		srcU += (vid->texWidth - vid->width) / 2;
 		srcV += (vid->texWidth - vid->width) / 2;
+		PRAGMA_DISABLE_DEPRECATION_WARNINGS
 		dstU += OutBuf->AllocatedPlaneDesc[1].ByteOffsetBetweenRows;
 		dstV += OutBuf->AllocatedPlaneDesc[2].ByteOffsetBetweenRows;
+		PRAGMA_ENABLE_DEPRECATION_WARNINGS
 	}
 }
 
 
+PRAGMA_DISABLE_DEPRECATION_WARNINGS
 FVideoDecoder::EDecodeResult FVideoDecoderMPEG4_Impl::Decode(const FVideoDecoderInput* InInput)
+PRAGMA_ENABLE_DEPRECATION_WARNINGS
 {
 	// With no registered callback that's interested in the result we can presume we would have been successful.
+	PRAGMA_DISABLE_DEPRECATION_WARNINGS
 	if (!OnDecodedFrame)
 	{
 		return FVideoDecoder::EDecodeResult::Success;
 	}
+	PRAGMA_ENABLE_DEPRECATION_WARNINGS
 	// Initialize decoder on first decode call.
 	if (!FirstUseInit())
 	{
+		PRAGMA_DISABLE_DEPRECATION_WARNINGS
 		return FVideoDecoder::EDecodeResult::Failure;
-	}
+	}PRAGMA_ENABLE_DEPRECATION_WARNINGS
+	
 
 	// Setup an access unit to run through the decoder.
 	TUniquePtr<FInDecoderData> AU = MakeUnique<FInDecoderData>();
+	PRAGMA_DISABLE_DEPRECATION_WARNINGS
 	if (InInput->GetDataSize() <= 0)
 	{
 		return FVideoDecoder::EDecodeResult::Failure;
@@ -404,6 +452,7 @@ FVideoDecoder::EDecodeResult FVideoDecoderMPEG4_Impl::Decode(const FVideoDecoder
 	AU->Height = InInput->GetHeight();
 	AU->bIsKeyframe = InInput->IsKeyframe();
 	AU->bIsComplete = InInput->IsCompleteFrame();
+	PRAGMA_ENABLE_DEPRECATION_WARNINGS
 	PendingDecodeData.Enqueue(MoveTemp(AU));
 
 	// Decode all pending input.
@@ -426,7 +475,9 @@ FVideoDecoder::EDecodeResult FVideoDecoderMPEG4_Impl::Decode(const FVideoDecoder
 			else
 			{
 				UE_LOG(LogVideoDecoder, Error, TEXT("VIDCreateDecoder() failed with %d"), (int)LastDecoderError);
+				PRAGMA_DISABLE_DEPRECATION_WARNINGS
 				return FVideoDecoder::EDecodeResult::Failure;
+				PRAGMA_ENABLE_DEPRECATION_WARNINGS
 				//vdecmpeg4::VIDDestroyDecoder(DecoderHandle);
 				//DecoderHandle = nullptr;
 			}
@@ -440,6 +491,7 @@ FVideoDecoder::EDecodeResult FVideoDecoderMPEG4_Impl::Decode(const FVideoDecoder
 				int32_t	Width = frame->width;
 				int32_t Height = frame->height;
 
+				PRAGMA_DISABLE_DEPRECATION_WARNINGS
 				TUniquePtr<FVideoDecoderOutputMPEG4> pNew(new FVideoDecoderOutputMPEG4(Width, Height, InInput->GetPTS()));
 
 				// Get memory from the application
@@ -482,13 +534,16 @@ FVideoDecoder::EDecodeResult FVideoDecoderMPEG4_Impl::Decode(const FVideoDecoder
 				{
 					pNew->SetPitchX(ap.Width);
 				}
+				PRAGMA_ENABLE_DEPRECATION_WARNINGS
 
 				// Copy the image across, converting it into NV12 format.
 				CopyI420ToNV12(pNew->GetBuffer(), frame);
 				frame->Release();
 				// Deliver
+				PRAGMA_DISABLE_DEPRECATION_WARNINGS
 				OnDecodedFrame(pNew.Release());
-			}
+				PRAGMA_ENABLE_DEPRECATION_WARNINGS
+			}	
 		}
 		else if (LastDecoderError == vdecmpeg4::VID_ERROR_STREAM_UNDERFLOW)
 		{
@@ -503,7 +558,9 @@ FVideoDecoder::EDecodeResult FVideoDecoderMPEG4_Impl::Decode(const FVideoDecoder
 			// Error!
 			UE_LOG(LogVideoDecoder, Error, TEXT("VIDStreamDecode() failed with %d"), (int)LastDecoderError);
 // TODO: destroy decoder and wait for next keyframe?
+			PRAGMA_DISABLE_DEPRECATION_WARNINGS
 			return FVideoDecoder::EDecodeResult::Failure;
+			PRAGMA_ENABLE_DEPRECATION_WARNINGS
 			break;
 		}
 
@@ -513,8 +570,11 @@ FVideoDecoder::EDecodeResult FVideoDecoderMPEG4_Impl::Decode(const FVideoDecoder
 			CurrentAU.Reset();
 		}
 	}
+	PRAGMA_DISABLE_DEPRECATION_WARNINGS
 	return FVideoDecoder::EDecodeResult::Success;
+	PRAGMA_ENABLE_DEPRECATION_WARNINGS
 }
 
+
 } // namespace AVEncoder
 
diff --git a/Engine/Source/Runtime/AVEncoder/Private/Decoders/VideoDecoderMpeg4.h b/Engine/Source/Runtime/AVEncoder/Private/Decoders/VideoDecoderMpeg4.h
--- a/Engine/Source/Runtime/AVEncoder/Private/Decoders/VideoDecoderMpeg4.h
+++ b/Engine/Source/Runtime/AVEncoder/Private/Decoders/VideoDecoderMpeg4.h
@@ -8,20 +8,26 @@
 namespace AVEncoder
 {
 
+PRAGMA_DISABLE_DEPRECATION_WARNINGS
 class FVideoDecoderMPEG4 : public FVideoDecoder
+PRAGMA_ENABLE_DEPRECATION_WARNINGS
 {
 public:
 	// register decoder with video decoder factory
+	PRAGMA_DISABLE_DEPRECATION_WARNINGS
 	static void Register(FVideoDecoderFactory& InFactory);
-
 	virtual bool Setup(const FInit& InInit) = 0;
+	virtual EDecodeResult Decode(const FVideoDecoderInput* InInput) = 0;
+	PRAGMA_ENABLE_DEPRECATION_WARNINGS
+
 	virtual void Shutdown() = 0;
 
-	virtual EDecodeResult Decode(const FVideoDecoderInput* InInput) = 0;
 
 protected:
+	PRAGMA_DISABLE_DEPRECATION_WARNINGS
 	virtual ~FVideoDecoderMPEG4() = default;
 	FVideoDecoderMPEG4() = default;
+	PRAGMA_ENABLE_DEPRECATION_WARNINGS
 };
 
 } // namespace AVEncoder
diff --git a/Engine/Source/Runtime/AVEncoder/Private/Decoders/VideoDecoderVP8.cpp b/Engine/Source/Runtime/AVEncoder/Private/Decoders/VideoDecoderVP8.cpp
--- a/Engine/Source/Runtime/AVEncoder/Private/Decoders/VideoDecoderVP8.cpp
+++ b/Engine/Source/Runtime/AVEncoder/Private/Decoders/VideoDecoderVP8.cpp
@@ -16,7 +16,9 @@
 namespace AVEncoder
 {
 
+PRAGMA_DISABLE_DEPRECATION_WARNINGS
 void FVideoDecoderVP8::Register(FVideoDecoderFactory& InFactory)
+PRAGMA_ENABLE_DEPRECATION_WARNINGS
 {
 }
 
diff --git a/Engine/Source/Runtime/AVEncoder/Private/Decoders/VideoDecoderVP8.h b/Engine/Source/Runtime/AVEncoder/Private/Decoders/VideoDecoderVP8.h
--- a/Engine/Source/Runtime/AVEncoder/Private/Decoders/VideoDecoderVP8.h
+++ b/Engine/Source/Runtime/AVEncoder/Private/Decoders/VideoDecoderVP8.h
@@ -8,16 +8,20 @@
 namespace AVEncoder
 {
 
+PRAGMA_DISABLE_DEPRECATION_WARNINGS
 class FVideoDecoderVP8 : public FVideoDecoder
+PRAGMA_ENABLE_DEPRECATION_WARNINGS
 {
 public:
 	// register decoder with video decoder factory
+	PRAGMA_DISABLE_DEPRECATION_WARNINGS
 	static void Register(FVideoDecoderFactory& InFactory);
-
 	virtual bool Setup(const FInit& InInit) = 0;
+	virtual EDecodeResult Decode(const FVideoDecoderInput* InInput) = 0;
+	PRAGMA_ENABLE_DEPRECATION_WARNINGS
+
 	virtual void Shutdown() = 0;
 
-	virtual EDecodeResult Decode(const FVideoDecoderInput* InInput) = 0;
 protected:
 	virtual ~FVideoDecoderVP8() = default;
 };
diff --git a/Engine/Source/Runtime/AVEncoder/Private/Decoders/Windows/VideoDecoderH264_Windows.cpp b/Engine/Source/Runtime/AVEncoder/Private/Decoders/Windows/VideoDecoderH264_Windows.cpp
--- a/Engine/Source/Runtime/AVEncoder/Private/Decoders/Windows/VideoDecoderH264_Windows.cpp
+++ b/Engine/Source/Runtime/AVEncoder/Private/Decoders/Windows/VideoDecoderH264_Windows.cpp
@@ -417,14 +417,20 @@ void FF5PlayerVideoDecoderOutputDX::ShutdownPoolable()
 /**
  * Decoded media sample
  */
+PRAGMA_DISABLE_DEPRECATION_WARNINGS
 class FVideoDecoderOutputH264_Windows : public FVideoDecoderOutput
+PRAGMA_ENABLE_DEPRECATION_WARNINGS
 {
 public:
+	PRAGMA_DISABLE_DEPRECATION_WARNINGS
 	FVideoDecoderOutputH264_Windows()
+	PRAGMA_ENABLE_DEPRECATION_WARNINGS
 	{}
 
+	PRAGMA_DISABLE_DEPRECATION_WARNINGS
 	virtual ~FVideoDecoderOutputH264_Windows()
 	{}
+	PRAGMA_ENABLE_DEPRECATION_WARNINGS
 
 	virtual int32 AddRef() override
 	{
@@ -480,7 +486,9 @@ public:
 		PTS = InPTS;
 	}
 
+	PRAGMA_DISABLE_DEPRECATION_WARNINGS
 	virtual const FVideoDecoderAllocFrameBufferResult* GetAllocatedBuffer() const
+	PRAGMA_ENABLE_DEPRECATION_WARNINGS
 	{
 		return &Buffer;
 	}
@@ -559,7 +567,9 @@ public:
 	}
 
 	// Internal for allocation.
+	PRAGMA_DISABLE_DEPRECATION_WARNINGS
 	FVideoDecoderAllocFrameBufferResult* GetBuffer()
+	PRAGMA_ENABLE_DEPRECATION_WARNINGS
 	{
 		return &Buffer;
 	}
@@ -571,7 +581,9 @@ public:
 	}
 
 private:
+	PRAGMA_DISABLE_DEPRECATION_WARNINGS
 	FVideoDecoderAllocFrameBufferResult	Buffer = {};
+	PRAGMA_ENABLE_DEPRECATION_WARNINGS
 	int32	RefCount = 1;
 	int32	Width = 0;
 	int32	Height = 0;
@@ -601,10 +613,14 @@ public:
 
 	virtual ~FVideoDecoderH264_WindowsImpl();
 
+	PRAGMA_DISABLE_DEPRECATION_WARNINGS
 	virtual bool Setup(const FInit& InInit) override;
+	PRAGMA_ENABLE_DEPRECATION_WARNINGS
 	virtual void Shutdown() override;
 
+	PRAGMA_DISABLE_DEPRECATION_WARNINGS
 	virtual EDecodeResult Decode(const FVideoDecoderInput* InInput) override;
+	PRAGMA_ENABLE_DEPRECATION_WARNINGS
 
 
 
@@ -805,16 +821,20 @@ bool						FVideoDecoderH264_WindowsImpl::bGlobalMFStartedUp = false;
 /**
  * Registers this decoder with the decoder factory.
  */
+PRAGMA_DISABLE_DEPRECATION_WARNINGS
 void FVideoDecoderH264_Windows::Register(FVideoDecoderFactory& InFactory)
+PRAGMA_ENABLE_DEPRECATION_WARNINGS
 {
+	PRAGMA_DISABLE_DEPRECATION_WARNINGS
 	FVideoDecoderInfo	DecoderInfo;
 	DecoderInfo.CodecType = ECodecType::H264;
 	DecoderInfo.MaxWidth = 1920;
 	DecoderInfo.MaxHeight = 1088;
-
+	
 	InFactory.Register(DecoderInfo, []() {
 			return new FVideoDecoderH264_WindowsImpl();
 		});
+	PRAGMA_ENABLE_DEPRECATION_WARNINGS
 }
 
 /*********************************************************************************************************************/
@@ -887,12 +907,18 @@ bool FVideoDecoderH264_WindowsImpl::CreateD3DResources()
 	void* ApplicationD3DDevice = nullptr;
 	int32_t ApplicationD3DDeviceVersion = 0;
 	bool bOk = false;
+	PRAGMA_DISABLE_DEPRECATION_WARNINGS
 	FVideoDecoderMethodsWindows* AllocationMethods = reinterpret_cast<FVideoDecoderMethodsWindows*>(GetAllocationInterfaceMethods());
+	PRAGMA_ENABLE_DEPRECATION_WARNINGS
 	if (AllocationMethods)
 	{
+		PRAGMA_DISABLE_DEPRECATION_WARNINGS
 		if (AllocationMethods->MagicCookie == 0x57696e58) // 'WinX'
+		PRAGMA_ENABLE_DEPRECATION_WARNINGS
 		{
+			PRAGMA_DISABLE_DEPRECATION_WARNINGS
 			int32_t AppResult = AllocationMethods->GetD3DDevice.Execute(AllocationMethods->This, &ApplicationD3DDevice, &ApplicationD3DDeviceVersion);
+			PRAGMA_ENABLE_DEPRECATION_WARNINGS
 			if (AppResult == 0)
 			{
 				// D3D version is returned as major*1000+minor.
@@ -919,7 +945,9 @@ bool FVideoDecoderH264_WindowsImpl::CreateD3DResources()
 		}
 		else
 		{
+			PRAGMA_DISABLE_DEPRECATION_WARNINGS
 			UE_LOG(LogVideoDecoder, Error, TEXT("Incorrect video decoder allocation interface is set. Is %08x, should be %08x"), AllocationMethods->MagicCookie, 0x57696e58);
+			PRAGMA_ENABLE_DEPRECATION_WARNINGS
 		}
 	}
 	else
@@ -1601,11 +1629,14 @@ bool FVideoDecoderH264_WindowsImpl::AllocateApplicationOutputSampleBuffer(int32
 
 	// Get the buffer for the output sample from the application. This is the part that gets wrapped in a
 	// webrtc structure and passed back out to the application.
+	PRAGMA_DISABLE_DEPRECATION_WARNINGS
 	FVideoDecoderAllocFrameBufferParams ap {};
 	EFrameBufferAllocReturn ar;
+	PRAGMA_ENABLE_DEPRECATION_WARNINGS
 	if (DecodeMode == EDecodeMode::Win8HW)
 	{
 		// We will pass a ID3D11Texture2D pointer to a texture of _our_ decoding device to _share_ with the renderer.
+		PRAGMA_DISABLE_DEPRECATION_WARNINGS
 		ap.FrameBufferType = EFrameBufferType::CODEC_TextureHandle;
 		ap.AllocSize = sizeof(ID3D11Texture2D*);
 		ap.AllocAlignment = alignof(ID3D11Texture2D*);
@@ -1614,10 +1645,12 @@ bool FVideoDecoderH264_WindowsImpl::AllocateApplicationOutputSampleBuffer(int32
 		ap.Width = Width;
 		ap.Height = Height;
 		ap.BytesPerPixel = 1;
+		PRAGMA_ENABLE_DEPRECATION_WARNINGS
 	}
 	else if (DecodeMode == EDecodeMode::Win10HW || DecodeMode == EDecodeMode::Win10SW || DecodeMode == EDecodeMode::XB1)
 	{
 		// We need a byte buffer to put the image into. On the application side this may be realized as a TArray<uint8>.
+		PRAGMA_DISABLE_DEPRECATION_WARNINGS
 		ap.FrameBufferType = EFrameBufferType::CODEC_RawBuffer;
 		ap.AllocSize = Width * Height * 3 / 2;
 		ap.AllocAlignment = 1;
@@ -1625,10 +1658,12 @@ bool FVideoDecoderH264_WindowsImpl::AllocateApplicationOutputSampleBuffer(int32
 		ap.Width = Width;
 		ap.Height = Height;
 		ap.BytesPerPixel = 1;
+		PRAGMA_ENABLE_DEPRECATION_WARNINGS
 	}
 	else if (DecodeMode == EDecodeMode::Win8SW)
 	{
 		// We want to have an ID3D11Texture2D we can give the decoder to decode into.
+		PRAGMA_DISABLE_DEPRECATION_WARNINGS
 		ap.FrameBufferType = EFrameBufferType::CODEC_TextureObject;
 		ap.AllocSize = sizeof(ID3D11Texture2D*);
 		ap.AllocAlignment = alignof(ID3D11Texture2D*);
@@ -1636,6 +1671,7 @@ bool FVideoDecoderH264_WindowsImpl::AllocateApplicationOutputSampleBuffer(int32
 		ap.Width = Width;
 		ap.Height = Height;
 		ap.BytesPerPixel = 1;
+		PRAGMA_ENABLE_DEPRECATION_WARNINGS
 	}
 	else
 	{
@@ -1643,14 +1679,20 @@ bool FVideoDecoderH264_WindowsImpl::AllocateApplicationOutputSampleBuffer(int32
 		return false;
 	}
 
+	PRAGMA_DISABLE_DEPRECATION_WARNINGS
 	ar = AllocateOutputFrameBuffer(CurrentApplicationOutputSample->GetBuffer(), &ap);
 	if (ar == EFrameBufferAllocReturn::CODEC_Success)
+	PRAGMA_ENABLE_DEPRECATION_WARNINGS
 	{
 		// Got an output buffer.
+		PRAGMA_DISABLE_DEPRECATION_WARNINGS
 		check(CurrentApplicationOutputSample->GetBuffer()->AllocatedBuffer);
+		PRAGMA_ENABLE_DEPRECATION_WARNINGS
 		return true;
 	}
+	PRAGMA_DISABLE_DEPRECATION_WARNINGS
 	else if (ar == EFrameBufferAllocReturn::CODEC_TryAgainLater)
+	PRAGMA_ENABLE_DEPRECATION_WARNINGS
 	{
 		UE_LOG(LogVideoDecoder, Error, TEXT("Unsupported allocation return value"));
 		return false;
@@ -1666,10 +1708,12 @@ void FVideoDecoderH264_WindowsImpl::ReturnUnusedApplicationOutputSample()
 {
 	if (CurrentApplicationOutputSample)
 	{
+		PRAGMA_DISABLE_DEPRECATION_WARNINGS
 		if (FVideoDecoderAllocFrameBufferResult* afb = CurrentApplicationOutputSample->GetBuffer())
 		{
 			afb->ReleaseCallback.ExecuteIfBound(afb->CallbackValue, afb->AllocatedBuffer);
 		}
+		PRAGMA_ENABLE_DEPRECATION_WARNINGS
 		CurrentApplicationOutputSample->Release();
 		CurrentApplicationOutputSample = nullptr;
 	}
@@ -1748,7 +1792,9 @@ bool FVideoDecoderH264_WindowsImpl::ProcessDecoding(FInputToDecode* InInput)
 			if (DecodeMode == EDecodeMode::Win8SW || DecodeMode == EDecodeMode::Win10SW)
 			{
 				// DX11 software mode needs a texture prepared by the application
+				PRAGMA_DISABLE_DEPRECATION_WARNINGS
 				if (DecodeMode == EDecodeMode::Win8SW && !CurrentApplicationOutputSample->GetBuffer()->AllocatedBuffer)
+				PRAGMA_ENABLE_DEPRECATION_WARNINGS
 				{
 					if (!AllocateApplicationOutputSampleBuffer(InWidth, InHeight))
 					{
@@ -1759,7 +1805,9 @@ bool FVideoDecoderH264_WindowsImpl::ProcessDecoding(FInputToDecode* InInput)
 				{
 					CurrentSoftwareDecoderOutput = new FF5PlayerVideoDecoderOutputDX;
 				}
+				PRAGMA_DISABLE_DEPRECATION_WARNINGS
 				if (!CurrentSoftwareDecoderOutput->PreInitForSoftwareDecode(FIntPoint(InWidth, InHeight), DecodeMode == EDecodeMode::Win8SW ? *static_cast<ID3D11Texture2D**>(CurrentApplicationOutputSample->GetBuffer()->AllocatedBuffer) : nullptr))
+				PRAGMA_ENABLE_DEPRECATION_WARNINGS
 				{
 					return false;
 				}
@@ -1864,7 +1912,9 @@ bool FVideoDecoderH264_WindowsImpl::CopyTexture(const TRefCountPtr<IMFSample>& D
 
 	if (DecodeMode == EDecodeMode::Win8HW)
 	{
+		PRAGMA_DISABLE_DEPRECATION_WARNINGS
 		check(CurrentApplicationOutputSample->GetBuffer()->AllocatedBuffer == nullptr);
+		PRAGMA_ENABLE_DEPRECATION_WARNINGS
 		if (!AllocateApplicationOutputSampleBuffer(OutputDim.X, OutputDim.Y))
 		{
 			UE_LOG(LogVideoDecoder, Error, TEXT("Failed to get a decode output buffer from the application!"));
@@ -1874,12 +1924,16 @@ bool FVideoDecoderH264_WindowsImpl::CopyTexture(const TRefCountPtr<IMFSample>& D
 	}
 	else if (DecodeMode == EDecodeMode::Win8SW)
 	{
+		PRAGMA_DISABLE_DEPRECATION_WARNINGS
 		check(CurrentApplicationOutputSample->GetBuffer()->AllocatedBuffer != nullptr);
+		PRAGMA_ENABLE_DEPRECATION_WARNINGS
 		DecoderOutput->InitializeWithTextureBuffer(OutputDim.X, FIntPoint(OutputDim.X, OutputDim.Y));// * 3 / 2));
 	}
 	else if (DecodeMode == EDecodeMode::Win10SW)
 	{
+		PRAGMA_DISABLE_DEPRECATION_WARNINGS
 		check(CurrentApplicationOutputSample->GetBuffer()->AllocatedBuffer == nullptr);
+		PRAGMA_ENABLE_DEPRECATION_WARNINGS
 		if (!AllocateApplicationOutputSampleBuffer(NativeDim.X, NativeDim.Y))
 		{
 			UE_LOG(LogVideoDecoder, Error, TEXT("Failed to get a decode output buffer from the application!"));
@@ -1906,8 +1960,10 @@ bool FVideoDecoderH264_WindowsImpl::CopyTexture(const TRefCountPtr<IMFSample>& D
 
 		// Copy NV12 texture data into external application buffer
 		check(CurrentApplicationOutputSample);
+		PRAGMA_DISABLE_DEPRECATION_WARNINGS
 		check((DWORD)CurrentApplicationOutputSample->GetBuffer()->AllocatedSize >= BufferSize);
 		FMemory::Memcpy(CurrentApplicationOutputSample->GetBuffer()->AllocatedBuffer, Data, BufferSize);
+		PRAGMA_ENABLE_DEPRECATION_WARNINGS
 
 		Result = Buffer->Unlock();
 		if (FAILED(Result))
@@ -1959,7 +2015,9 @@ bool FVideoDecoderH264_WindowsImpl::CopyTexture(const TRefCountPtr<IMFSample>& D
 			return false;
 		}
 
+		PRAGMA_DISABLE_DEPRECATION_WARNINGS
 		check(CurrentApplicationOutputSample->GetBuffer()->AllocatedBuffer == nullptr);
+		PRAGMA_ENABLE_DEPRECATION_WARNINGS
 		if (!AllocateApplicationOutputSampleBuffer(Pitch, TextureDesc.Height))
 		{
 			UE_LOG(LogVideoDecoder, Error, TEXT("Failed to get a decode output buffer from the application!"));
@@ -1974,8 +2032,10 @@ bool FVideoDecoderH264_WindowsImpl::CopyTexture(const TRefCountPtr<IMFSample>& D
 		DecoderOutput->InitializeWithBuffer(Pitch, FIntPoint(TextureDesc.Width, TextureDesc.Height * 3 / 2));
 
 		// Copy NV12 texture data into external application buffer
+		PRAGMA_DISABLE_DEPRECATION_WARNINGS
 		check((DWORD)CurrentApplicationOutputSample->GetBuffer()->AllocatedSize >= BufferSize);
 		FMemory::Memcpy(CurrentApplicationOutputSample->GetBuffer()->AllocatedBuffer, Data, BufferSize);
+		PRAGMA_ENABLE_DEPRECATION_WARNINGS
 
 		Result = Buffer2D->Unlock2D();
 		if (FAILED(Result))
@@ -2050,7 +2110,9 @@ bool FVideoDecoderH264_WindowsImpl::ConvertDecodedImage(const TRefCountPtr<IMFSa
 	CurrentApplicationOutputSample->SetNativeDecoderOutput(DecoderOutput);
 	if (DecodeMode == EDecodeMode::Win8HW)
 	{
+		PRAGMA_DISABLE_DEPRECATION_WARNINGS
 		ID3D11Texture2D** ApplicationTexturePtr = reinterpret_cast<ID3D11Texture2D**>(CurrentApplicationOutputSample->GetBuffer()->AllocatedBuffer);
+		PRAGMA_ENABLE_DEPRECATION_WARNINGS
 		// We set the texture in the user provided buffer. Note that we do NOT add the ref count here. This is ok because we have the
 		// DecoderOutput stored in CurrentApplicationOutputSample where it is held. If the user application code uses the texture it
 		// will/shall attach it to a ref counter pointer to retain it.
@@ -2141,10 +2203,12 @@ void FVideoDecoderH264_WindowsImpl::PurgeAllPendingOutput()
 	FVideoDecoderOutputH264_Windows* Output = nullptr;
 	while(OutputQueue.Dequeue(Output))
 	{
+		PRAGMA_DISABLE_DEPRECATION_WARNINGS
 		if (FVideoDecoderAllocFrameBufferResult* afb = Output->GetBuffer())
 		{
 			afb->ReleaseCallback.ExecuteIfBound(afb->CallbackValue, afb->AllocatedBuffer);
 		}
+		PRAGMA_ENABLE_DEPRECATION_WARNINGS
 		Output->Release();
 		Output = nullptr;
 	}
@@ -2186,11 +2250,15 @@ FVideoDecoderH264_WindowsImpl::~FVideoDecoderH264_WindowsImpl()
 /**
  * Called right after the constructor with initialization parameters.
  */
+PRAGMA_DISABLE_DEPRECATION_WARNINGS
 bool FVideoDecoderH264_WindowsImpl::Setup(const FInit& InInit)
+PRAGMA_ENABLE_DEPRECATION_WARNINGS
 {
 	// Remember the decoder allocation interface factory methods.
+	PRAGMA_DISABLE_DEPRECATION_WARNINGS
 	CreateDecoderAllocationInterfaceFN = InInit.CreateDecoderAllocationInterface;
 	ReleaseDecoderAllocationInterfaceFN = InInit.ReleaseDecoderAllocationInterface;
+	PRAGMA_ENABLE_DEPRECATION_WARNINGS
 	// [...]
 	CreateWorkerThread();
 	return true;
@@ -2204,7 +2272,9 @@ void FVideoDecoderH264_WindowsImpl::PerformAsyncShutdown()
 	if (bIsInitialized)
 	{
 		ReleaseD3DResources();
+		PRAGMA_DISABLE_DEPRECATION_WARNINGS
 		ReleaseDecoderAllocationInterface();
+		PRAGMA_ENABLE_DEPRECATION_WARNINGS
 		StaticReleaseResources();
 		bIsInitialized = false;
 	}
@@ -2231,7 +2301,9 @@ bool FVideoDecoderH264_WindowsImpl::SetupFirstUseResources()
 	if (!bIsInitialized)
 	{
 		if (StaticInitializeResources() &&
+			PRAGMA_DISABLE_DEPRECATION_WARNINGS
 			CreateDecoderAllocationInterface() &&
+			PRAGMA_ENABLE_DEPRECATION_WARNINGS
 			CreateD3DResources())
 		{
 			bIsInitialized = true;
@@ -2245,25 +2317,31 @@ bool FVideoDecoderH264_WindowsImpl::SetupFirstUseResources()
 /**
  * Called from webrtc indirectly to decode a frame.
  */
+PRAGMA_DISABLE_DEPRECATION_WARNINGS
 FVideoDecoder::EDecodeResult FVideoDecoderH264_WindowsImpl::Decode(const FVideoDecoderInput* InInput)
+PRAGMA_ENABLE_DEPRECATION_WARNINGS
 {
 	bool bOk = SetupFirstUseResources();
 	if (bOk)
 	{
+		PRAGMA_DISABLE_DEPRECATION_WARNINGS
 		if (!InInput->GetData() || InInput->GetDataSize() <= 0)
 		{
 			return FVideoDecoder::EDecodeResult::Failure;
 		}
+		PRAGMA_ENABLE_DEPRECATION_WARNINGS
 
 		// Need to wait for the decoder thread to have spun up??
 
 		TUniquePtr<FInputToDecode> Input = MakeUnique<FInputToDecode>();
+		PRAGMA_DISABLE_DEPRECATION_WARNINGS
 		Input->PTS = InInput->GetPTS();
 		Input->Width = InInput->GetWidth();
 		Input->Height = InInput->GetHeight();
 		Input->bIsKeyframe = InInput->IsKeyframe();
 		Input->Data.AddUninitialized(InInput->GetDataSize());
 		FMemory::Memcpy(Input->Data.GetData(), InInput->GetData(), InInput->GetDataSize());
+		PRAGMA_ENABLE_DEPRECATION_WARNINGS
 		InputQueue.Enqueue(MoveTemp(Input));
 
 		WorkerThreadSignalHaveWork.Signal();
@@ -2274,18 +2352,29 @@ FVideoDecoder::EDecodeResult FVideoDecoderH264_WindowsImpl::Decode(const FVideoD
 		{
 			if (Output->GetIsValid())
 			{
+				PRAGMA_DISABLE_DEPRECATION_WARNINGS
 				OnDecodedFrame(Output);
+				PRAGMA_ENABLE_DEPRECATION_WARNINGS
 			}
+			
 			else
 			{
 				delete Output;
+				PRAGMA_DISABLE_DEPRECATION_WARNINGS
 				return FVideoDecoder::EDecodeResult::Failure;
+				PRAGMA_ENABLE_DEPRECATION_WARNINGS
 			}
+			
 			Output = nullptr;
 		}
+		PRAGMA_DISABLE_DEPRECATION_WARNINGS
 		return FVideoDecoder::EDecodeResult::Success;
+		PRAGMA_ENABLE_DEPRECATION_WARNINGS
 	}
+	
+	PRAGMA_DISABLE_DEPRECATION_WARNINGS
 	return FVideoDecoder::EDecodeResult::Failure;
+	PRAGMA_ENABLE_DEPRECATION_WARNINGS
 }
 
 
diff --git a/Engine/Source/Runtime/AVEncoder/Private/Decoders/Windows/VideoDecoderH264_Windows.h b/Engine/Source/Runtime/AVEncoder/Private/Decoders/Windows/VideoDecoderH264_Windows.h
--- a/Engine/Source/Runtime/AVEncoder/Private/Decoders/Windows/VideoDecoderH264_Windows.h
+++ b/Engine/Source/Runtime/AVEncoder/Private/Decoders/Windows/VideoDecoderH264_Windows.h
@@ -12,12 +12,14 @@ class FVideoDecoderH264_Windows : public FVideoDecoderH264_Impl
 {
 public:
 	// register decoder with video decoder factory
+	PRAGMA_DISABLE_DEPRECATION_WARNINGS
 	static void Register(FVideoDecoderFactory& InFactory);
-
 	virtual bool Setup(const FInit& InInit) = 0;
+	virtual EDecodeResult Decode(const FVideoDecoderInput* InInput) = 0;
+	PRAGMA_ENABLE_DEPRECATION_WARNINGS
+
 	virtual void Shutdown() = 0;
 
-	virtual EDecodeResult Decode(const FVideoDecoderInput* InInput) = 0;
 protected:
 	virtual ~FVideoDecoderH264_Windows() = default;
 };
diff --git a/Engine/Source/Runtime/AVEncoder/Private/Encoders/WmfAudioEncoder.cpp b/Engine/Source/Runtime/AVEncoder/Private/Encoders/WmfAudioEncoder.cpp
--- a/Engine/Source/Runtime/AVEncoder/Private/Encoders/WmfAudioEncoder.cpp
+++ b/Engine/Source/Runtime/AVEncoder/Private/Encoders/WmfAudioEncoder.cpp
@@ -12,7 +12,9 @@
 
 namespace AVEncoder
 {
+PRAGMA_DISABLE_DEPRECATION_WARNINGS
 class FWmfAudioEncoder : public FAudioEncoder
+PRAGMA_ENABLE_DEPRECATION_WARNINGS
 {
 public:
 
@@ -26,10 +28,12 @@ public:
 
 	const TCHAR* GetName() const override;
 	const TCHAR* GetType() const override;
+	PRAGMA_DISABLE_DEPRECATION_WARNINGS
 	bool Initialize(const FAudioConfig& InConfig) override;
-	void Shutdown() override;
 	void Encode(const FAudioFrame& Frame) override;
 	FAudioConfig GetConfig() const override;
+	PRAGMA_ENABLE_DEPRECATION_WARNINGS
+	void Shutdown() override;
 
 private:
 	bool SetInputType();
@@ -37,14 +41,18 @@ private:
 	bool RetrieveStreamInfo();
 	bool StartStreaming();
 
+	PRAGMA_DISABLE_DEPRECATION_WARNINGS
 	FIMFSampleWrapper CreateInputSample(const uint8* SampleData, uint32 Size, FTimespan Timestamp, FTimespan Duration);
 	FIMFSampleWrapper GetOutputSample();
+	PRAGMA_ENABLE_DEPRECATION_WARNINGS
 
 	ECodecType CodecType;
 	FString Name;
 	FString Type;
 
+	PRAGMA_DISABLE_DEPRECATION_WARNINGS
 	FAudioConfig Config;
+	PRAGMA_ENABLE_DEPRECATION_WARNINGS
 	TRefCountPtr<IMFTransform> Encoder;
 	TRefCountPtr<IMFMediaType> OutputType;
 	MFT_INPUT_STREAM_INFO InputStreamInfo = {};
@@ -58,7 +66,7 @@ private:
 //////////////////////////////////////////////////////////////////////////
 // FWmfAudioEncoder implementation
 //////////////////////////////////////////////////////////////////////////
-
+PRAGMA_DISABLE_DEPRECATION_WARNINGS
 FWmfAudioEncoder::FWmfAudioEncoder(ECodecType CodecType)
 	: CodecType(CodecType)
 {
@@ -72,10 +80,13 @@ FWmfAudioEncoder::FWmfAudioEncoder(ECodecType CodecType)
 		checkNoEntry();
 	}
 }
+PRAGMA_ENABLE_DEPRECATION_WARNINGS
 
+PRAGMA_DISABLE_DEPRECATION_WARNINGS
 FWmfAudioEncoder::~FWmfAudioEncoder()
 {
 }
+PRAGMA_ENABLE_DEPRECATION_WARNINGS
 
 const TCHAR* FWmfAudioEncoder::GetName() const
 {
@@ -116,7 +127,9 @@ namespace
 	}
 }
 
+PRAGMA_DISABLE_DEPRECATION_WARNINGS
 bool FWmfAudioEncoder::Initialize(const FAudioConfig& InConfig)
+PRAGMA_ENABLE_DEPRECATION_WARNINGS
 {
 
 	//
@@ -126,10 +139,12 @@ bool FWmfAudioEncoder::Initialize(const FAudioConfig& InConfig)
 	{
 		// See  https://docs.microsoft.com/en-us/windows/desktop/medfound/aac-encoder for details
 		FString ErrorStr;
+		PRAGMA_DISABLE_DEPRECATION_WARNINGS
 		if (
 			!ValidateValue(TEXT("AAC Bitrate"), InConfig.Bitrate, {12000*8, 16000*8, 20000*8, 24000*8}, ErrorStr) ||
 			!ValidateValue(TEXT("AAC Samplerate"), InConfig.Samplerate, { 44100, 48000, 0 }, ErrorStr) ||
 			!ValidateValue(TEXT("AAC NumChannels"), InConfig.NumChannels, {1,2,6}, ErrorStr))
+		PRAGMA_ENABLE_DEPRECATION_WARNINGS
 		{
 			UE_LOG(LogAVEncoder, Error, TEXT("%s"), *ErrorStr);
 			return false;
@@ -141,9 +156,11 @@ bool FWmfAudioEncoder::Initialize(const FAudioConfig& InConfig)
 		return false;
 	}
 
+	PRAGMA_DISABLE_DEPRECATION_WARNINGS
 	UE_LOG(LogAVEncoder, Log, TEXT("AudioEncoder config: %d channels, %d Hz, %.2f Kbps"), InConfig.NumChannels, InConfig.Samplerate, InConfig.Bitrate / 1000.0f);
 
 	Config = InConfig;
+	PRAGMA_ENABLE_DEPRECATION_WARNINGS
 
 	const GUID* CodecGuid=nullptr;
 	if (CodecType == ECodecType::AAC)
@@ -171,6 +188,7 @@ bool FWmfAudioEncoder::Initialize(const FAudioConfig& InConfig)
 	return true;
 }
 
+
 bool FWmfAudioEncoder::SetInputType()
 {
 	TRefCountPtr<IMFMediaType> MediaType;
@@ -178,8 +196,11 @@ bool FWmfAudioEncoder::SetInputType()
 	CHECK_HR(MediaType->SetGUID(MF_MT_MAJOR_TYPE, MFMediaType_Audio));
 	CHECK_HR(MediaType->SetGUID(MF_MT_SUBTYPE, MFAudioFormat_PCM));
 	CHECK_HR(MediaType->SetUINT32(MF_MT_AUDIO_BITS_PER_SAMPLE, 16)); // the only value supported
+
+	PRAGMA_DISABLE_DEPRECATION_WARNINGS
 	CHECK_HR(MediaType->SetUINT32(MF_MT_AUDIO_SAMPLES_PER_SECOND, Config.Samplerate));
 	CHECK_HR(MediaType->SetUINT32(MF_MT_AUDIO_NUM_CHANNELS, Config.NumChannels));
+	PRAGMA_DISABLE_DEPRECATION_WARNINGS
 
 	CHECK_HR(Encoder->SetInputType(0, MediaType, 0));
 
diff --git a/Engine/Source/Runtime/AVEncoder/Private/Encoders/WmfAudioEncoder.h b/Engine/Source/Runtime/AVEncoder/Private/Encoders/WmfAudioEncoder.h
--- a/Engine/Source/Runtime/AVEncoder/Private/Encoders/WmfAudioEncoder.h
+++ b/Engine/Source/Runtime/AVEncoder/Private/Encoders/WmfAudioEncoder.h
@@ -10,14 +10,18 @@
 
 namespace AVEncoder
 {
+	PRAGMA_DISABLE_DEPRECATION_WARNINGS
 	class FWmfAudioEncoderFactory : public FAudioEncoderFactory
+	PRAGMA_ENABLE_DEPRECATION_WARNINGS
 	{
 	public:
 		FWmfAudioEncoderFactory();
 		~FWmfAudioEncoderFactory() override;
 		const TCHAR* GetName() const override;
 		TArray<FString> GetSupportedCodecs() const override;
+		PRAGMA_DISABLE_DEPRECATION_WARNINGS
 		TUniquePtr<FAudioEncoder> CreateEncoder(const FString& Codec) override;
+		PRAGMA_ENABLE_DEPRECATION_WARNINGS
 	private:
 	};
 }
diff --git a/Engine/Source/Runtime/AVEncoder/Private/VideoDecoderFactory.cpp b/Engine/Source/Runtime/AVEncoder/Private/VideoDecoderFactory.cpp
--- a/Engine/Source/Runtime/AVEncoder/Private/VideoDecoderFactory.cpp
+++ b/Engine/Source/Runtime/AVEncoder/Private/VideoDecoderFactory.cpp
@@ -23,11 +23,13 @@ namespace AVEncoder
 FThreadSafeCounter			FVideoDecoderFactory::NextID = 1;
 bool						FVideoDecoderFactory::bDebugDontRegisterDefaultCodecs = false;
 
+PRAGMA_DISABLE_DEPRECATION_WARNINGS
 FVideoDecoderFactory& FVideoDecoderFactory::Get()
 {
 	static FVideoDecoderFactory Singleton;
 	return Singleton;
 }
+PRAGMA_ENABLE_DEPRECATION_WARNINGS
 
 FVideoDecoderFactory::FVideoDecoderFactory()
 {
@@ -50,10 +52,14 @@ void FVideoDecoderFactory::Debug_SetDontRegisterDefaultCodecs()
 }
 
 
+PRAGMA_DISABLE_DEPRECATION_WARNINGS
 void FVideoDecoderFactory::Register(const FVideoDecoderInfo& InInfo, const CreateDecoderCallback& InCreateEncoder)
+PRAGMA_ENABLE_DEPRECATION_WARNINGS
 {
 	AvailableDecoders.Push(InInfo);
+	PRAGMA_DISABLE_DEPRECATION_WARNINGS
 	AvailableDecoders.Last().ID = NextID.Increment();
+	PRAGMA_ENABLE_DEPRECATION_WARNINGS
 	CreateDecoders.Push(InCreateEncoder);
 }
 
@@ -75,11 +81,15 @@ void FVideoDecoderFactory::RegisterDefaultCodecs()
 #endif
 }
 
+PRAGMA_DISABLE_DEPRECATION_WARNINGS
 bool FVideoDecoderFactory::GetInfo(uint32 InID, FVideoDecoderInfo& OutInfo) const
+PRAGMA_ENABLE_DEPRECATION_WARNINGS
 {
 	for (int32 Index = 0; Index < AvailableDecoders.Num(); ++Index)
 	{
+		PRAGMA_DISABLE_DEPRECATION_WARNINGS
 		if (AvailableDecoders[Index].ID == InID)
+		PRAGMA_ENABLE_DEPRECATION_WARNINGS
 		{
 			OutInfo = AvailableDecoders[Index];
 			return true;
@@ -88,20 +98,27 @@ bool FVideoDecoderFactory::GetInfo(uint32 InID, FVideoDecoderInfo& OutInfo) cons
 	return false;
 }
 
-
+PRAGMA_DISABLE_DEPRECATION_WARNINGS
 FVideoDecoder* FVideoDecoderFactory::Create(uint32 InID, const FVideoDecoder::FInit& InInit)
+PRAGMA_ENABLE_DEPRECATION_WARNINGS
 {
+	PRAGMA_DISABLE_DEPRECATION_WARNINGS
 	FVideoDecoder* Result = nullptr;
+	PRAGMA_ENABLE_DEPRECATION_WARNINGS
 	for (int32 Index = 0; Index < AvailableDecoders.Num(); ++Index)
 	{
+		PRAGMA_DISABLE_DEPRECATION_WARNINGS
 		if (AvailableDecoders[Index].ID == InID)
+		PRAGMA_ENABLE_DEPRECATION_WARNINGS
 		{
 			Result = CreateDecoders[Index]();
+			PRAGMA_DISABLE_DEPRECATION_WARNINGS
 			if (Result && !Result->Setup(InInit))
 			{
 				Result->Shutdown();
 				Result = nullptr;
 			}
+			PRAGMA_ENABLE_DEPRECATION_WARNINGS
 			break;
 		}
 	}
diff --git a/Engine/Source/Runtime/AVEncoder/Private/VideoDecoderInput.cpp b/Engine/Source/Runtime/AVEncoder/Private/VideoDecoderInput.cpp
--- a/Engine/Source/Runtime/AVEncoder/Private/VideoDecoderInput.cpp
+++ b/Engine/Source/Runtime/AVEncoder/Private/VideoDecoderInput.cpp
@@ -5,11 +5,15 @@
 namespace AVEncoder
 {
 
+PRAGMA_DISABLE_DEPRECATION_WARNINGS
 class FVideoDecoderInputImpl : public FVideoDecoderInput
 {
+PRAGMA_ENABLE_DEPRECATION_WARNINGS
 public:
+	PRAGMA_DISABLE_DEPRECATION_WARNINGS
 	virtual ~FVideoDecoderInputImpl() = default;
 	FVideoDecoderInputImpl(const FInputData& InInputData);
+	PRAGMA_ENABLE_DEPRECATION_WARNINGS
 
 	virtual int32 GetWidth() const override
 	{ return Width; }
@@ -45,12 +49,14 @@ private:
 	bool		bMissingFrames;
 };
 
+PRAGMA_DISABLE_DEPRECATION_WARNINGS
 TSharedPtr<FVideoDecoderInput> FVideoDecoderInput::Create(const FInputData& InInputData)
 {
 	return MakeShared<FVideoDecoderInputImpl>(InInputData);
 }
+PRAGMA_ENABLE_DEPRECATION_WARNINGS
 
-
+PRAGMA_DISABLE_DEPRECATION_WARNINGS
 FVideoDecoderInputImpl::FVideoDecoderInputImpl(const FInputData& InInputData)
 {
 	AccessUnit = InInputData.EncodedData;
@@ -64,6 +70,6 @@ FVideoDecoderInputImpl::FVideoDecoderInputImpl(const FInputData& InInputData)
 	bIsCompleteFrame = InInputData.bIsComplete;
 	bMissingFrames = InInputData.bMissingFrames;
 }
-
+PRAGMA_ENABLE_DEPRECATION_WARNINGS
 
 } /* namespace AVEncoder */
diff --git a/Engine/Source/Runtime/AVEncoder/Private/VideoEncoderInput.cpp b/Engine/Source/Runtime/AVEncoder/Private/VideoEncoderInput.cpp
--- a/Engine/Source/Runtime/AVEncoder/Private/VideoEncoderInput.cpp
+++ b/Engine/Source/Runtime/AVEncoder/Private/VideoEncoderInput.cpp
@@ -30,8 +30,9 @@ namespace AVEncoder
 	// *** FVideoEncoderInput *************************************************************************
 
 	// --- construct video encoder input based on expected input frame format -------------------------
-
+	PRAGMA_DISABLE_DEPRECATION_WARNINGS
 	TSharedPtr<FVideoEncoderInput> FVideoEncoderInput::CreateDummy(bool bIsResizable)
+	PRAGMA_ENABLE_DEPRECATION_WARNINGS
 	{
 		TSharedPtr<FVideoEncoderInputImpl>	Input = MakeShared<FVideoEncoderInputImpl>();
 		Input->bIsResizable = bIsResizable;
@@ -43,7 +44,9 @@ namespace AVEncoder
 		return Input;
 	}
 
+	PRAGMA_DISABLE_DEPRECATION_WARNINGS
 	TSharedPtr<FVideoEncoderInput> FVideoEncoderInput::CreateForYUV420P(uint32 InWidth, uint32 InHeight, bool bIsResizable)
+	PRAGMA_ENABLE_DEPRECATION_WARNINGS
 	{
 		TSharedPtr<FVideoEncoderInputImpl>	Input = MakeShared<FVideoEncoderInputImpl>();
 		Input->bIsResizable = bIsResizable;
@@ -55,7 +58,9 @@ namespace AVEncoder
 		return Input;
 	}
 
+	PRAGMA_DISABLE_DEPRECATION_WARNINGS
 	TSharedPtr<FVideoEncoderInput> FVideoEncoderInput::CreateForD3D11(void* InApplicationD3DDevice, bool bIsResizable, bool IsShared)
+	PRAGMA_ENABLE_DEPRECATION_WARNINGS
 	{
 		TSharedPtr<FVideoEncoderInputImpl>	Input = MakeShared<FVideoEncoderInputImpl>();
 
@@ -83,7 +88,9 @@ namespace AVEncoder
 	}
 
 
+	PRAGMA_DISABLE_DEPRECATION_WARNINGS
 	TSharedPtr<FVideoEncoderInput> FVideoEncoderInput::CreateForD3D12(void* InApplicationD3DDevice, bool bIsResizable, bool IsShared)
+	PRAGMA_ENABLE_DEPRECATION_WARNINGS
 	{
 		TSharedPtr<FVideoEncoderInputImpl>	Input = MakeShared<FVideoEncoderInputImpl>();
 
@@ -110,7 +117,9 @@ namespace AVEncoder
 		return Input;
 	}
 
+	PRAGMA_DISABLE_DEPRECATION_WARNINGS
 	TSharedPtr<FVideoEncoderInput> FVideoEncoderInput::CreateForCUDA(void* InApplicationContext, bool bIsResizable)
+	PRAGMA_ENABLE_DEPRECATION_WARNINGS
 	{
 		TSharedPtr<FVideoEncoderInputImpl> Input = MakeShared<FVideoEncoderInputImpl>();
 		Input->bIsResizable = bIsResizable;
@@ -123,14 +132,20 @@ namespace AVEncoder
 	}
 
 #if PLATFORM_DESKTOP && !PLATFORM_APPLE
+	PRAGMA_DISABLE_DEPRECATION_WARNINGS
 	TSharedPtr<FVideoEncoderInput> FVideoEncoderInput::CreateForVulkan(void* InApplicationVulkanData, bool bIsResizable)
+	PRAGMA_ENABLE_DEPRECATION_WARNINGS
 	{
 		TSharedPtr<FVideoEncoderInputImpl>	Input = MakeShared<FVideoEncoderInputImpl>();
 		Input->bIsResizable = bIsResizable;
 
+		PRAGMA_DISABLE_DEPRECATION_WARNINGS
 		FVulkanDataStruct* VulkanData = static_cast<FVulkanDataStruct*>(InApplicationVulkanData);
+		PRAGMA_ENABLE_DEPRECATION_WARNINGS
 
+		PRAGMA_DISABLE_DEPRECATION_WARNINGS
 		if (!Input->SetupForVulkan(VulkanData->VulkanInstance, VulkanData->VulkanPhysicalDevice, VulkanData->VulkanDevice))
+		PRAGMA_ENABLE_DEPRECATION_WARNINGS
 		{
 			Input.Reset();
 		}
@@ -151,6 +166,7 @@ namespace AVEncoder
 
 	// *** FVideoEncoderInputImpl *********************************************************************
 
+	PRAGMA_DISABLE_DEPRECATION_WARNINGS
 	FVideoEncoderInputImpl::~FVideoEncoderInputImpl()
 	{
 		{
@@ -166,16 +182,21 @@ namespace AVEncoder
 		//	DEBUG_D3D11_REPORT_LIVE_DEVICE_OBJECT(FrameInfoD3D.EncoderDeviceD3D11);
 #endif
 	}
+	PRAGMA_ENABLE_DEPRECATION_WARNINGS
 
 	bool FVideoEncoderInputImpl::SetupForDummy()
 	{
+		PRAGMA_DISABLE_DEPRECATION_WARNINGS
 		FrameFormat = EVideoFrameFormat::Undefined;
+		PRAGMA_ENABLE_DEPRECATION_WARNINGS
 		return true;
 	}
 
 	bool FVideoEncoderInputImpl::SetupForYUV420P(uint32 InWidth, uint32 InHeight)
 	{
+		PRAGMA_DISABLE_DEPRECATION_WARNINGS
 		FrameFormat = EVideoFrameFormat::YUV420P;
+		PRAGMA_ENABLE_DEPRECATION_WARNINGS
 		FrameInfoYUV420P.StrideY = InWidth;
 		FrameInfoYUV420P.StrideU = (InWidth + 1) / 2;
 		FrameInfoYUV420P.StrideV = (InWidth + 1) / 2;
@@ -224,7 +245,9 @@ namespace AVEncoder
 		DEBUG_SET_D3D11_OBJECT_NAME(FrameInfoD3D.EncoderDeviceD3D11, "FVideoEncoderInputImpl");
 		DEBUG_SET_D3D11_OBJECT_NAME(FrameInfoD3D.EncoderDeviceContextD3D11, "FVideoEncoderInputImpl");
 
+		PRAGMA_DISABLE_DEPRECATION_WARNINGS
 		FrameFormat = EVideoFrameFormat::D3D11_R8G8B8A8_UNORM;
+		PRAGMA_ENABLE_DEPRECATION_WARNINGS
 
 		CollectAvailableEncoders();
 		return true;
@@ -274,7 +297,9 @@ namespace AVEncoder
 		DEBUG_SET_D3D11_OBJECT_NAME(FrameInfoD3D.EncoderDeviceD3D11, "FVideoEncoderInputImpl");
 		DEBUG_SET_D3D11_OBJECT_NAME(FrameInfoD3D.EncoderDeviceContextD3D11, "FVideoEncoderInputImpl");
 
+		PRAGMA_DISABLE_DEPRECATION_WARNINGS
 		FrameFormat = EVideoFrameFormat::D3D11_R8G8B8A8_UNORM;
+		PRAGMA_ENABLE_DEPRECATION_WARNINGS
 
 		CollectAvailableEncoders();
 		return true;
@@ -311,7 +336,9 @@ namespace AVEncoder
 			return false;
 		}
 
+		PRAGMA_DISABLE_DEPRECATION_WARNINGS
 		FrameFormat = EVideoFrameFormat::D3D12_R8G8B8A8_UNORM;
+		PRAGMA_ENABLE_DEPRECATION_WARNINGS
 
 		CollectAvailableEncoders();
 		return true;
@@ -371,7 +398,9 @@ namespace AVEncoder
 		DEBUG_SET_D3D11_OBJECT_NAME(FrameInfoD3D.EncoderDeviceD3D11, "FVideoEncoderInputImpl");
 		DEBUG_SET_D3D11_OBJECT_NAME(FrameInfoD3D.EncoderDeviceContextD3D11, "FVideoEncoderInputImpl");
 
+		PRAGMA_DISABLE_DEPRECATION_WARNINGS
 		FrameFormat = EVideoFrameFormat::D3D11_R8G8B8A8_UNORM;
+		PRAGMA_ENABLE_DEPRECATION_WARNINGS
 
 		CollectAvailableEncoders();
 		return true;
@@ -386,7 +415,9 @@ namespace AVEncoder
 	{
 		FrameInfoCUDA.EncoderContextCUDA = static_cast<CUcontext>(InApplicationContext);
 
+		PRAGMA_DISABLE_DEPRECATION_WARNINGS
 		FrameFormat = EVideoFrameFormat::CUDA_R8G8B8A8_UNORM;
+		PRAGMA_ENABLE_DEPRECATION_WARNINGS
 
 		CollectAvailableEncoders();
 		return true;
@@ -395,11 +426,13 @@ namespace AVEncoder
 #if PLATFORM_DESKTOP && !PLATFORM_APPLE
 	bool FVideoEncoderInputImpl::SetupForVulkan(VkInstance InApplicationVulkanInstance, VkPhysicalDevice InApplicationVulkanPhysicalDevice, VkDevice InApplicationVulkanDevice)
 	{
+		PRAGMA_DISABLE_DEPRECATION_WARNINGS
 		FrameInfoVulkan.VulkanInstance = InApplicationVulkanInstance;
 		FrameInfoVulkan.VulkanPhysicalDevice = InApplicationVulkanPhysicalDevice;
 		FrameInfoVulkan.VulkanDevice = InApplicationVulkanDevice;
 
 		FrameFormat = EVideoFrameFormat::VULKAN_R8G8B8A8_UNORM;
+		PRAGMA_ENABLE_DEPRECATION_WARNINGS
 
 		CollectAvailableEncoders();
 		return true;
@@ -411,6 +444,7 @@ namespace AVEncoder
 	void FVideoEncoderInputImpl::CollectAvailableEncoders()
 	{
 		AvailableEncoders.Empty();
+		PRAGMA_DISABLE_DEPRECATION_WARNINGS
 		for (const FVideoEncoderInfo& Info : FVideoEncoderFactory::Get().GetAvailable())
 		{
 			if (Info.SupportedInputFormats.Contains(FrameFormat))
@@ -418,16 +452,21 @@ namespace AVEncoder
 				AvailableEncoders.Push(Info);
 			}
 		}
+		PRAGMA_ENABLE_DEPRECATION_WARNINGS
 	}
 
+	PRAGMA_DISABLE_DEPRECATION_WARNINGS
 	const TArray<FVideoEncoderInfo>& FVideoEncoderInputImpl::GetAvailableEncoders()
+	PRAGMA_ENABLE_DEPRECATION_WARNINGS
 	{
 		return AvailableEncoders;
 	}
 
 	// --- encoder input frames -----------------------------------------------------------------------
 
+	PRAGMA_DISABLE_DEPRECATION_WARNINGS
 	bool FVideoEncoderInputImpl::IsUserManagedFrame(const FVideoEncoderInputFrame* InBuffer) const
+	PRAGMA_ENABLE_DEPRECATION_WARNINGS
 	{
 		const FVideoEncoderInputFrameImpl* Frame = static_cast<const FVideoEncoderInputFrameImpl*>(InBuffer);
 		FScopeLock						Guard(&ProtectFrames);
@@ -442,7 +481,9 @@ namespace AVEncoder
 	}
 
 	// create a user managed buffer
+	PRAGMA_DISABLE_DEPRECATION_WARNINGS
 	FVideoEncoderInputFrame* FVideoEncoderInputImpl::CreateBuffer(OnFrameReleasedCallback InOnFrameReleased)
+	PRAGMA_ENABLE_DEPRECATION_WARNINGS
 	{
 		FVideoEncoderInputFrameImpl* Frame = CreateFrame();
 		if (Frame)
@@ -454,7 +495,9 @@ namespace AVEncoder
 	}
 
 	// destroy user managed buffer
+	PRAGMA_DISABLE_DEPRECATION_WARNINGS
 	void FVideoEncoderInputImpl::DestroyBuffer(FVideoEncoderInputFrame* InBuffer)
+	PRAGMA_ENABLE_DEPRECATION_WARNINGS
 	{
 		FVideoEncoderInputFrameImpl* Frame = static_cast<FVideoEncoderInputFrameImpl*>(InBuffer);
 		FScopeLock						Guard(&ProtectFrames);
@@ -475,7 +518,9 @@ namespace AVEncoder
 
 	// --- encoder input frames -----------------------------------------------------------------------
 
+	PRAGMA_DISABLE_DEPRECATION_WARNINGS
 	TSharedPtr<FVideoEncoderInputFrame> FVideoEncoderInputImpl::ObtainInputFrame()
+	PRAGMA_ENABLE_DEPRECATION_WARNINGS
 	{
 		TSharedPtr<FVideoEncoderInputFrameImpl> Frame;
 		FScopeLock						Guard(&ProtectFrames);
@@ -489,11 +534,15 @@ namespace AVEncoder
 		else
 		{
 			Frame = MakeShareable(CreateFrame());
+			PRAGMA_DISABLE_DEPRECATION_WARNINGS
 			UE_LOG(LogVideoEncoder, Verbose, TEXT("Created new frame total frames: %d"), NumBuffers);
+			PRAGMA_ENABLE_DEPRECATION_WARNINGS
 		}
+		
 
 		ActiveFrames.Push(Frame);
 
+		PRAGMA_DISABLE_DEPRECATION_WARNINGS
 		Frame->SetFrameID(NextFrameID++);
 
 		if (NextFrameID == 0)
@@ -502,12 +551,14 @@ namespace AVEncoder
 		}
 
 		Frame->Obtain();
+		PRAGMA_ENABLE_DEPRECATION_WARNINGS
 		return Frame;
 	}
 
 	FVideoEncoderInputFrameImpl* FVideoEncoderInputImpl::CreateFrame()
 	{
 		FVideoEncoderInputFrameImpl* Frame = new FVideoEncoderInputFrameImpl(this);
+		PRAGMA_DISABLE_DEPRECATION_WARNINGS
 		NumBuffers++;
 		switch (FrameFormat)
 		{
@@ -533,10 +584,13 @@ namespace AVEncoder
 			check(false);
 			break;
 		}
+		PRAGMA_ENABLE_DEPRECATION_WARNINGS
 		return Frame;
 	}
 
+	PRAGMA_DISABLE_DEPRECATION_WARNINGS
 	void FVideoEncoderInputImpl::ReleaseInputFrame(FVideoEncoderInputFrame* InFrame)
+	PRAGMA_ENABLE_DEPRECATION_WARNINGS
 	{
 		FVideoEncoderInputFrameImpl* InFrameImpl = static_cast<FVideoEncoderInputFrameImpl*>(InFrame);
 
@@ -564,6 +618,7 @@ namespace AVEncoder
 		if (NumRemoved > 0)
 		{
 			// drop frame if format changed
+			PRAGMA_DISABLE_DEPRECATION_WARNINGS
 			if (InFrame->GetFormat() != FrameFormat)
 			{
 				ProtectFrames.Unlock();
@@ -579,6 +634,7 @@ namespace AVEncoder
 				UE_LOG(LogVideoEncoder, Verbose, TEXT("Deleted buffer (too many) total frames: %d"), NumBuffers);
 				return;
 			}
+			PRAGMA_ENABLE_DEPRECATION_WARNINGS
 
 			AvailableFrames.Enqueue(InFramePtr);
 		}
@@ -588,32 +644,40 @@ namespace AVEncoder
 	{
 		FScopeLock ScopeLock(&ProtectFrames);
 		AvailableFrames.Empty();
+		PRAGMA_DISABLE_DEPRECATION_WARNINGS
 		NumBuffers = 0;
+		PRAGMA_ENABLE_DEPRECATION_WARNINGS
 	}
+	
 
 	void FVideoEncoderInputImpl::SetupFrameYUV420P(FVideoEncoderInputFrameImpl* Frame)
 	{
+		PRAGMA_DISABLE_DEPRECATION_WARNINGS
 		Frame->SetFormat(EVideoFrameFormat::YUV420P);
 		FVideoEncoderInputFrame::FYUV420P& YUV420P = Frame->GetYUV420P();
 		YUV420P.StrideY = FrameInfoYUV420P.StrideY;
 		YUV420P.StrideU = FrameInfoYUV420P.StrideU;
 		YUV420P.StrideV = FrameInfoYUV420P.StrideV;
 		YUV420P.Data[0] = YUV420P.Data[1] = YUV420P.Data[2] = nullptr;
+		PRAGMA_ENABLE_DEPRECATION_WARNINGS
 	}
 
 	void FVideoEncoderInputImpl::SetupFrameD3D11(FVideoEncoderInputFrameImpl* Frame)
 	{
 #if PLATFORM_WINDOWS
+		PRAGMA_DISABLE_DEPRECATION_WARNINGS
 		Frame->SetFormat(FrameFormat);
 
 		FVideoEncoderInputFrame::FD3D11& Data = Frame->GetD3D11();
 		Data.EncoderDevice = FrameInfoD3D.EncoderDeviceD3D11;
+		PRAGMA_ENABLE_DEPRECATION_WARNINGS
 #endif
 	}
 
 	void FVideoEncoderInputImpl::SetupFrameD3D12(FVideoEncoderInputFrameImpl* Frame)
 	{
 #if PLATFORM_WINDOWS
+		PRAGMA_DISABLE_DEPRECATION_WARNINGS
 		Frame->SetFormat(FrameFormat);
 
 		if (FrameFormat == EVideoFrameFormat::D3D11_R8G8B8A8_UNORM)
@@ -626,24 +690,30 @@ namespace AVEncoder
 			FVideoEncoderInputFrame::FD3D12& Data = Frame->GetD3D12();
 			Data.EncoderDevice = FrameInfoD3D.EncoderDeviceD3D12;
 		}
+		PRAGMA_ENABLE_DEPRECATION_WARNINGS
 #endif
 	}
 
 	void FVideoEncoderInputImpl::SetupFrameVulkan(FVideoEncoderInputFrameImpl* Frame)
 	{
 #if PLATFORM_DESKTOP && !PLATFORM_APPLE
+		PRAGMA_DISABLE_DEPRECATION_WARNINGS
 		Frame->SetFormat(FrameFormat);
 
 		FVideoEncoderInputFrame::FVulkan& Data = Frame->GetVulkan();
 		Data.EncoderDevice = FrameInfoVulkan.VulkanDevice;
+		PRAGMA_ENABLE_DEPRECATION_WARNINGS
 #endif
 	}
 
 	void FVideoEncoderInputImpl::SetupFrameCUDA(FVideoEncoderInputFrameImpl* Frame)
 	{
+		PRAGMA_DISABLE_DEPRECATION_WARNINGS
 		Frame->SetFormat(FrameFormat);
+
 		FVideoEncoderInputFrame::FCUDA& Data = Frame->GetCUDA();
 		Data.EncoderDevice = FrameInfoCUDA.EncoderContextCUDA;
+		PRAGMA_ENABLE_DEPRECATION_WARNINGS
 	}
 
 	// ---
@@ -679,24 +749,29 @@ namespace AVEncoder
 		, TimestampUs(0)
 		, TimestampRTP(0)
 		, NumReferences(0)
-		, Format(EVideoFrameFormat::Undefined)
 		, Width(0)
 		, Height(0)
 		, bFreeYUV420PData(false)
+		PRAGMA_DISABLE_DEPRECATION_WARNINGS
+		, Format(EVideoFrameFormat::Undefined)
+		PRAGMA_ENABLE_DEPRECATION_WARNINGS
 	{
 	}
 
+	PRAGMA_DISABLE_DEPRECATION_WARNINGS
 	FVideoEncoderInputFrame::FVideoEncoderInputFrame(const FVideoEncoderInputFrame& CloneFrom)
+	PRAGMA_ENABLE_DEPRECATION_WARNINGS
 		: FrameID(CloneFrom.FrameID)
 		, TimestampUs(CloneFrom.TimestampUs)
 		, TimestampRTP(CloneFrom.TimestampRTP)
 		, NumReferences(0)
-		, Format(CloneFrom.Format)
 		, Width(CloneFrom.Width)
 		, Height(CloneFrom.Height)
 		, bFreeYUV420PData(false)
+		, Format(CloneFrom.Format)
 	{
 #if PLATFORM_WINDOWS
+		PRAGMA_DISABLE_DEPRECATION_WARNINGS
 		D3D11.EncoderDevice = CloneFrom.D3D11.EncoderDevice;
 		D3D11.Texture = CloneFrom.D3D11.Texture;
 		if ((D3D11.EncoderTexture = CloneFrom.D3D11.EncoderTexture) != nullptr)
@@ -710,37 +785,49 @@ namespace AVEncoder
 		{
 			D3D12.EncoderTexture->AddRef();
 		}
+		PRAGMA_ENABLE_DEPRECATION_WARNINGS
 #endif
 
+		PRAGMA_DISABLE_DEPRECATION_WARNINGS
 		CUDA.EncoderDevice = CloneFrom.CUDA.EncoderDevice;
 		CUDA.EncoderTexture = CloneFrom.CUDA.EncoderTexture;
+		PRAGMA_ENABLE_DEPRECATION_WARNINGS
 	}
 
 	FVideoEncoderInputFrame::~FVideoEncoderInputFrame()
 	{
 		if (bFreeYUV420PData)
 		{
+			PRAGMA_DISABLE_DEPRECATION_WARNINGS
 			delete[] YUV420P.Data[0];
 			delete[] YUV420P.Data[1];
 			delete[] YUV420P.Data[2];
 			YUV420P.Data[0] = YUV420P.Data[1] = YUV420P.Data[2] = nullptr;
+			PRAGMA_ENABLE_DEPRECATION_WARNINGS
 			bFreeYUV420PData = false;
 		}
 
 #if PLATFORM_WINDOWS
+		PRAGMA_DISABLE_DEPRECATION_WARNINGS
 		if (D3D11.EncoderTexture)
+		PRAGMA_ENABLE_DEPRECATION_WARNINGS
 		{
 			// check to make sure this frame holds the last reference
+			PRAGMA_DISABLE_DEPRECATION_WARNINGS
 			auto	NumRef = D3D11.EncoderTexture->AddRef();
+			PRAGMA_ENABLE_DEPRECATION_WARNINGS
 			if (NumRef > 2)
 			{
 				UE_LOG(LogVideoEncoder, Warning, TEXT("VideoEncoderFame - D3D11 input texture still holds %d references."), NumRef);
 			}
 			// Need to call release twice as we have added an extra reference just above (only way to count how many references we have)
+			PRAGMA_DISABLE_DEPRECATION_WARNINGS
 			D3D11.EncoderTexture->Release();
 			D3D11.EncoderTexture->Release();
 			D3D11.EncoderTexture = nullptr;
+			PRAGMA_ENABLE_DEPRECATION_WARNINGS
 		}
+		PRAGMA_DISABLE_DEPRECATION_WARNINGS
 		if (D3D11.SharedHandle)
 		{
 			CloseHandle(D3D11.SharedHandle);
@@ -768,22 +855,28 @@ namespace AVEncoder
 			OnReleaseD3D12Texture(D3D12.Texture);
 			D3D12.Texture = nullptr;
 		}
+		PRAGMA_ENABLE_DEPRECATION_WARNINGS
 
 		// D3D12 specific handle atm
+		PRAGMA_DISABLE_DEPRECATION_WARNINGS
 		if (CUDA.SharedHandle)
 		{
 			CloseHandle(CUDA.SharedHandle);
 			CUDA.SharedHandle = nullptr;
 		}
+		PRAGMA_ENABLE_DEPRECATION_WARNINGS
 #endif 
 
+		PRAGMA_DISABLE_DEPRECATION_WARNINGS
 		if (CUDA.EncoderTexture)
 		{
 			OnReleaseCUDATexture(CUDA.EncoderTexture);
 			CUDA.EncoderTexture = nullptr;
 		}
+		PRAGMA_ENABLE_DEPRECATION_WARNINGS
 
 #if PLATFORM_DESKTOP && !PLATFORM_APPLE
+		PRAGMA_DISABLE_DEPRECATION_WARNINGS
 		if (Vulkan.EncoderTexture != VK_NULL_HANDLE)
 		{
 			OnReleaseVulkanTexture(Vulkan.EncoderTexture);
@@ -793,6 +886,7 @@ namespace AVEncoder
 		{
 			OnReleaseVulkanSurface(Vulkan.EncoderSurface);
 		}
+		PRAGMA_ENABLE_DEPRECATION_WARNINGS
 #endif
 	}
 
@@ -800,45 +894,59 @@ namespace AVEncoder
 	{
 		if (!bFreeYUV420PData)
 		{
+			PRAGMA_DISABLE_DEPRECATION_WARNINGS
 			YUV420P.StrideY = Width;
 			YUV420P.StrideU = (Width + 1) / 2;
 			YUV420P.StrideV = (Width + 1) / 2;;
 			YUV420P.Data[0] = new uint8[Height * YUV420P.StrideY];
 			YUV420P.Data[1] = new uint8[(Height + 1) / 2 * YUV420P.StrideU];
 			YUV420P.Data[2] = new uint8[(Height + 1) / 2 * YUV420P.StrideV];
+			PRAGMA_ENABLE_DEPRECATION_WARNINGS
 			bFreeYUV420PData = true;
 		}
 	}
 
 	void FVideoEncoderInputFrame::SetYUV420P(const uint8* InDataY, const uint8* InDataU, const uint8* InDataV, uint32 InStrideY, uint32 InStrideU, uint32 InStrideV)
 	{
+		PRAGMA_DISABLE_DEPRECATION_WARNINGS
 		if (Format == EVideoFrameFormat::YUV420P)
+		PRAGMA_ENABLE_DEPRECATION_WARNINGS
 		{
 			if (bFreeYUV420PData)
 			{
+				PRAGMA_DISABLE_DEPRECATION_WARNINGS
 				delete[] YUV420P.Data[0];
 				delete[] YUV420P.Data[1];
 				delete[] YUV420P.Data[2];
+				PRAGMA_ENABLE_DEPRECATION_WARNINGS
 				bFreeYUV420PData = false;
 			}
+			PRAGMA_DISABLE_DEPRECATION_WARNINGS
 			YUV420P.Data[0] = InDataY;
 			YUV420P.Data[1] = InDataU;
 			YUV420P.Data[2] = InDataV;
 			YUV420P.StrideY = InStrideY;
 			YUV420P.StrideU = InStrideU;
 			YUV420P.StrideV = InStrideV;
+			PRAGMA_ENABLE_DEPRECATION_WARNINGS
 		}
 	}
 
 	static FThreadSafeCounter	_VideoEncoderInputFrameCnt{ 0 };
 
 #if PLATFORM_WINDOWS
+	PRAGMA_DISABLE_DEPRECATION_WARNINGS
 	void FVideoEncoderInputFrame::SetTexture(ID3D11Texture2D* InTexture, FReleaseD3D11TextureCallback InOnReleaseTexture)
+	PRAGMA_ENABLE_DEPRECATION_WARNINGS
 	{
+		PRAGMA_DISABLE_DEPRECATION_WARNINGS
 		if (Format == EVideoFrameFormat::D3D11_R8G8B8A8_UNORM)
+		PRAGMA_ENABLE_DEPRECATION_WARNINGS
 		{
+			PRAGMA_DISABLE_DEPRECATION_WARNINGS
 			check(D3D11.Texture == nullptr);
 			if (!D3D11.Texture)
+			PRAGMA_ENABLE_DEPRECATION_WARNINGS
 			{
 				TRefCountPtr<IDXGIResource> DXGIResource;
 				HANDLE						SharedHandle;
@@ -858,28 +966,38 @@ namespace AVEncoder
 				{
 					UE_LOG(LogVideoEncoder, Error, TEXT("IDXGIResource::GetSharedHandle() failed to return a shared texture resource no created as shared? (D3D11_RESOURCE_MISC_SHARED)."));
 				}
+				PRAGMA_DISABLE_DEPRECATION_WARNINGS
 				else if ((Result = D3D11.EncoderDevice->OpenSharedResource(SharedHandle, __uuidof(ID3D11Texture2D), (LPVOID*)&D3D11.EncoderTexture)) != S_OK)
+				PRAGMA_ENABLE_DEPRECATION_WARNINGS
 				{
 					UE_LOG(LogVideoEncoder, Error, TEXT("ID3D11Device::OpenSharedResource() failed 0x%X - %s."), Result, *GetComErrorDescription(Result));
 				}
 				else
 				{
+					PRAGMA_DISABLE_DEPRECATION_WARNINGS
 					DEBUG_SET_D3D11_OBJECT_NAME(D3D11.EncoderTexture, "FVideoEncoderInputFrame::SetTexture()");
 					D3D11.Texture = InTexture;
+					PRAGMA_ENABLE_DEPRECATION_WARNINGS
 					OnReleaseD3D11Texture = InOnReleaseTexture;
 				}
 			}
 		}
 	}
 
+	PRAGMA_DISABLE_DEPRECATION_WARNINGS
 	void FVideoEncoderInputFrame::SetTexture(ID3D12Resource* InTexture, FReleaseD3D12TextureCallback InOnReleaseTexture)
+	PRAGMA_ENABLE_DEPRECATION_WARNINGS
 	{
+		PRAGMA_DISABLE_DEPRECATION_WARNINGS
 		if (Format == EVideoFrameFormat::D3D11_R8G8B8A8_UNORM)
+		PRAGMA_ENABLE_DEPRECATION_WARNINGS
 		{
+			PRAGMA_DISABLE_DEPRECATION_WARNINGS
 			check(D3D12.Texture == nullptr)
-				check(D3D11.EncoderTexture == nullptr)
+			check(D3D11.EncoderTexture == nullptr)
+			PRAGMA_ENABLE_DEPRECATION_WARNINGS
 
-				TRefCountPtr<ID3D12Device>	OwnerDevice;
+			TRefCountPtr<ID3D12Device>	OwnerDevice;
 			HRESULT						Result;
 			if ((Result = InTexture->GetDevice(IID_PPV_ARGS(OwnerDevice.GetInitReference()))) != S_OK)
 			{
@@ -888,48 +1006,65 @@ namespace AVEncoder
 			//
 			// NOTE: ID3D12Device::CreateSharedHandle gives as an NT Handle, and so we need to call CloseHandle on it
 			//
+			PRAGMA_DISABLE_DEPRECATION_WARNINGS
 			else if ((Result = OwnerDevice->CreateSharedHandle(InTexture, NULL, GENERIC_ALL, *FString::Printf(TEXT("%s_FVideoEncoderInputFrame_%d"), *GetGUID(), _VideoEncoderInputFrameCnt.Increment()), &D3D11.SharedHandle)) != S_OK)
+			PRAGMA_ENABLE_DEPRECATION_WARNINGS
 			{
 				UE_LOG(LogVideoEncoder, Error, TEXT("ID3D12Device::CreateSharedHandle() failed 0x%X - %s."), Result, *GetComErrorDescription(Result));
 			}
+			PRAGMA_DISABLE_DEPRECATION_WARNINGS
 			else if (D3D11.SharedHandle == nullptr)
+			PRAGMA_ENABLE_DEPRECATION_WARNINGS
 			{
 				UE_LOG(LogVideoEncoder, Error, TEXT("ID3D12Device::CreateSharedHandle() failed to return a shared texture resource no created as shared? (D3D12_RESOURCE_FLAG_ALLOW_SIMULTANEOUS_ACCESS)."));
 			}
 			else
 			{
 				TRefCountPtr <ID3D11Device1> Device1;
+				PRAGMA_DISABLE_DEPRECATION_WARNINGS
 				if ((Result = D3D11.EncoderDevice->QueryInterface(IID_PPV_ARGS(Device1.GetInitReference()))) != S_OK)
+				PRAGMA_ENABLE_DEPRECATION_WARNINGS
 				{
 					UE_LOG(LogVideoEncoder, Error, TEXT("ID3D11Device::QueryInterface() failed 0x%X - %s."), Result, *GetComErrorDescription(Result));
 				}
+				PRAGMA_DISABLE_DEPRECATION_WARNINGS
 				else if ((Result = Device1->OpenSharedResource1(D3D11.SharedHandle, IID_PPV_ARGS(&D3D11.EncoderTexture))) != S_OK)
+				PRAGMA_ENABLE_DEPRECATION_WARNINGS
 				{
 					UE_LOG(LogVideoEncoder, Error, TEXT("ID3D11Device::OpenSharedResource1() failed 0x%X - %s."), Result, *GetComErrorDescription(Result));
 				}
 				else
 				{
+					PRAGMA_DISABLE_DEPRECATION_WARNINGS
 					DEBUG_SET_D3D11_OBJECT_NAME(D3D11.EncoderTexture, "FVideoEncoderInputFrame::SetTexture()");
 					D3D12.Texture = InTexture;
+					PRAGMA_ENABLE_DEPRECATION_WARNINGS
 					OnReleaseD3D12Texture = InOnReleaseTexture;
 				}
 			}
 		}
 		else
 		{
+			PRAGMA_DISABLE_DEPRECATION_WARNINGS
 			check(D3D12.Texture == nullptr);
 			check(D3D12.EncoderTexture == nullptr);
 			D3D12.Texture = InTexture;
 			D3D12.EncoderTexture = InTexture;
+			PRAGMA_ENABLE_DEPRECATION_WARNINGS
 			OnReleaseD3D12Texture = InOnReleaseTexture;
 		}
 	}
 #endif
 
+	PRAGMA_DISABLE_DEPRECATION_WARNINGS
 	void FVideoEncoderInputFrame::SetTexture(CUarray InTexture, EUnderlyingRHI UnderlyingRHI, void* SharedHandle, FReleaseCUDATextureCallback InOnReleaseTexture)
+	PRAGMA_ENABLE_DEPRECATION_WARNINGS
 	{
+		PRAGMA_DISABLE_DEPRECATION_WARNINGS
 		if (Format == EVideoFrameFormat::CUDA_R8G8B8A8_UNORM)
+		PRAGMA_ENABLE_DEPRECATION_WARNINGS
 		{
+			PRAGMA_DISABLE_DEPRECATION_WARNINGS
 			CUDA.UnderlyingRHI = UnderlyingRHI;
 			CUDA.SharedHandle = SharedHandle;
 			CUDA.EncoderTexture = InTexture;
@@ -938,28 +1073,40 @@ namespace AVEncoder
 			{
 				UE_LOG(LogVideoEncoder, Warning, TEXT("SetTexture | Cuda device pointer is null"));
 			}
+			PRAGMA_ENABLE_DEPRECATION_WARNINGS
 		}
 	}
 
 
 #if PLATFORM_DESKTOP && !PLATFORM_APPLE
+	PRAGMA_DISABLE_DEPRECATION_WARNINGS
 	void FVideoEncoderInputFrame::SetTexture(VkImage InTexture, FReleaseVulkanTextureCallback InOnReleaseTexture)
+	PRAGMA_ENABLE_DEPRECATION_WARNINGS
 	{
+		PRAGMA_DISABLE_DEPRECATION_WARNINGS
 		if (Format == EVideoFrameFormat::VULKAN_R8G8B8A8_UNORM)
+		PRAGMA_ENABLE_DEPRECATION_WARNINGS
 		{
+			PRAGMA_DISABLE_DEPRECATION_WARNINGS
 			Vulkan.EncoderTexture = InTexture;
 			OnReleaseVulkanTexture = InOnReleaseTexture;
 			if (!Vulkan.EncoderTexture)
 			{
 				UE_LOG(LogVideoEncoder, Warning, TEXT("SetTexture | Vulkan VkImage is null"));
 			}
+			PRAGMA_ENABLE_DEPRECATION_WARNINGS
 		}
 	}
 
+	PRAGMA_DISABLE_DEPRECATION_WARNINGS
 	void FVideoEncoderInputFrame::SetTexture(VkImage InTexture, VkDeviceMemory InTextureDeviceMemory, uint64 InTextureMemorySize, FReleaseVulkanTextureCallback InOnReleaseTexture)
+	PRAGMA_ENABLE_DEPRECATION_WARNINGS
 	{
+		PRAGMA_DISABLE_DEPRECATION_WARNINGS
 		if (Format == EVideoFrameFormat::VULKAN_R8G8B8A8_UNORM)
+		PRAGMA_ENABLE_DEPRECATION_WARNINGS
 		{
+			PRAGMA_DISABLE_DEPRECATION_WARNINGS
 			Vulkan.EncoderTexture = InTexture;
 			Vulkan.EncoderDeviceMemory = InTextureDeviceMemory;
 			Vulkan.EncoderMemorySize = InTextureMemorySize;
@@ -969,17 +1116,21 @@ namespace AVEncoder
 			{
 				UE_LOG(LogVideoEncoder, Warning, TEXT("SetTexture | Vulkan VkImage or VkTextureDeviceMemory is null"));
 			}
+			PRAGMA_ENABLE_DEPRECATION_WARNINGS
 		}
 	}
 #endif
 
 	// *** FVideoEncoderInputFrameImpl ****************************************************************
 
+	PRAGMA_DISABLE_DEPRECATION_WARNINGS
 	FVideoEncoderInputFrameImpl::FVideoEncoderInputFrameImpl(FVideoEncoderInputImpl* InInput)
 		: Input(InInput)
 	{
 	}
+	PRAGMA_ENABLE_DEPRECATION_WARNINGS
 
+	PRAGMA_DISABLE_DEPRECATION_WARNINGS
 	FVideoEncoderInputFrameImpl::FVideoEncoderInputFrameImpl(const FVideoEncoderInputFrameImpl& InCloneFrom, FCloneDestroyedCallback InCloneDestroyedCallback)
 		: FVideoEncoderInputFrame(InCloneFrom)
 		, Input(InCloneFrom.Input)
@@ -987,7 +1138,9 @@ namespace AVEncoder
 		, OnCloneDestroyed(MoveTemp(InCloneDestroyedCallback))
 	{
 	}
+	PRAGMA_ENABLE_DEPRECATION_WARNINGS
 
+	PRAGMA_DISABLE_DEPRECATION_WARNINGS
 	FVideoEncoderInputFrameImpl::~FVideoEncoderInputFrameImpl()
 	{
 		if (ClonedReference)
@@ -998,6 +1151,7 @@ namespace AVEncoder
 		{
 		}
 	}
+	PRAGMA_ENABLE_DEPRECATION_WARNINGS
 
 	void FVideoEncoderInputFrameImpl::Release() const
 	{
@@ -1006,7 +1160,9 @@ namespace AVEncoder
 		{
 			Input->ReleaseInputFrame(const_cast<FVideoEncoderInputFrameImpl*>(this));
 		}
+		PRAGMA_DISABLE_DEPRECATION_WARNINGS
 		else if (NumReferences.Decrement() == 0)
+		PRAGMA_ENABLE_DEPRECATION_WARNINGS
 		{
 			if (ClonedReference)
 			{
@@ -1021,12 +1177,11 @@ namespace AVEncoder
 	}
 
 	// Clone frame - this will create a copy that references the original until destroyed
+	PRAGMA_DISABLE_DEPRECATION_WARNINGS
 	const FVideoEncoderInputFrame* FVideoEncoderInputFrameImpl::Clone(FCloneDestroyedCallback InCloneDestroyedCallback) const
+	PRAGMA_ENABLE_DEPRECATION_WARNINGS
 	{
 		FVideoEncoderInputFrameImpl* ClonedFrame = new FVideoEncoderInputFrameImpl(*this, MoveTemp(InCloneDestroyedCallback));
 		return ClonedFrame;
 	}
-
-
-
 } /* namespace AVEncoder */
diff --git a/Engine/Source/Runtime/AVEncoder/Private/VideoEncoderInputImpl.h b/Engine/Source/Runtime/AVEncoder/Private/VideoEncoderInputImpl.h
--- a/Engine/Source/Runtime/AVEncoder/Private/VideoEncoderInputImpl.h
+++ b/Engine/Source/Runtime/AVEncoder/Private/VideoEncoderInputImpl.h
@@ -17,6 +17,7 @@ namespace AVEncoder
 {
 class FVideoEncoderInputFrameImpl;
 
+PRAGMA_DISABLE_DEPRECATION_WARNINGS
 class FVideoEncoderInputImpl : public FVideoEncoderInput
 {
 public:
@@ -165,5 +166,5 @@ private:
 	const FVideoEncoderInputFrame*			ClonedReference = nullptr;
 	const FCloneDestroyedCallback			OnCloneDestroyed;
 };
-
+PRAGMA_ENABLE_DEPRECATION_WARNINGS
 } /* namespace AVEncoder */
diff --git a/Engine/Source/Runtime/AVEncoder/Private/VideoEncoderModule.cpp b/Engine/Source/Runtime/AVEncoder/Private/VideoEncoderModule.cpp
--- a/Engine/Source/Runtime/AVEncoder/Private/VideoEncoderModule.cpp
+++ b/Engine/Source/Runtime/AVEncoder/Private/VideoEncoderModule.cpp
@@ -11,7 +11,9 @@ public:
 
 	void ShutdownModule()
 	{
+		PRAGMA_DISABLE_DEPRECATION_WARNINGS
 		AVEncoder::FVideoEncoderFactory::Shutdown();
+		PRAGMA_ENABLE_DEPRECATION_WARNINGS
 	}
 };
 
diff --git a/Engine/Source/Runtime/AVEncoder/Public/AudioEncoder.h b/Engine/Source/Runtime/AVEncoder/Public/AudioEncoder.h
--- a/Engine/Source/Runtime/AVEncoder/Public/AudioEncoder.h
+++ b/Engine/Source/Runtime/AVEncoder/Public/AudioEncoder.h
@@ -12,44 +12,55 @@
 
 namespace AVEncoder
 {
-	struct FAudioFrame
+	struct UE_DEPRECATED(5.4, "AVEncoder has been deprecated. Please use the AVCodecs plugin family instead.") FAudioFrame
 	{
 		FTimespan Timestamp;
 		FTimespan Duration;
 		Audio::TSampleBuffer<float> Data;
 	};
 
-	class IAudioEncoderListener
+	class UE_DEPRECATED(5.4, "AVEncoder has been deprecated. Please use the AVCodecs plugin family instead.") IAudioEncoderListener
 	{
 	public:
+		PRAGMA_DISABLE_DEPRECATION_WARNINGS
 		virtual void OnEncodedAudioFrame(const FMediaPacket& Packet) = 0;
+		PRAGMA_ENABLE_DEPRECATION_WARNINGS
 	};
 
-	class FAudioEncoder
+	class UE_DEPRECATED(5.4, "AVEncoder has been deprecated. Please use the AVCodecs plugin family instead.") FAudioEncoder
 	{
 	public:
 		virtual ~FAudioEncoder() {}
 		virtual const TCHAR* GetName() const = 0;
 		virtual const TCHAR* GetType() const = 0;
+		PRAGMA_DISABLE_DEPRECATION_WARNINGS
 		virtual bool Initialize(const FAudioConfig& Config) = 0;
+		PRAGMA_ENABLE_DEPRECATION_WARNINGS
 
 		/**
 		* Shutdown MUST be called before destruction
 		*/
 		virtual void Shutdown() = 0;
 
+		PRAGMA_DISABLE_DEPRECATION_WARNINGS
 		virtual void Encode(const FAudioFrame& Frame) = 0;
 		virtual FAudioConfig GetConfig() const = 0;
 
 		AVENCODER_API virtual void RegisterListener(IAudioEncoderListener& Listener);
 		AVENCODER_API virtual void UnregisterListener(IAudioEncoderListener& Listener);
+		PRAGMA_ENABLE_DEPRECATION_WARNINGS
 
 	protected:
+		PRAGMA_DISABLE_DEPRECATION_WARNINGS
 		AVENCODER_API void OnEncodedAudioFrame(const FMediaPacket& Packet);
+		PRAGMA_ENABLE_DEPRECATION_WARNINGS
 	private:
 
 		FCriticalSection ListenersMutex;
+
+		PRAGMA_DISABLE_DEPRECATION_WARNINGS
 		TArray<IAudioEncoderListener*> Listeners;
+		PRAGMA_ENABLE_DEPRECATION_WARNINGS
 
 	};
 }
diff --git a/Engine/Source/Runtime/AVEncoder/Public/AudioEncoderFactory.h b/Engine/Source/Runtime/AVEncoder/Public/AudioEncoderFactory.h
--- a/Engine/Source/Runtime/AVEncoder/Public/AudioEncoderFactory.h
+++ b/Engine/Source/Runtime/AVEncoder/Public/AudioEncoderFactory.h
@@ -11,13 +11,15 @@ namespace AVEncoder
 {
 	class FAudioEncoder;
 
-	class FAudioEncoderFactory
+	class UE_DEPRECATED(5.4, "AVEncoder has been deprecated. Please use the AVCodecs plugin family instead.") FAudioEncoderFactory
 	{
 	public:
 		virtual ~FAudioEncoderFactory() {}
 		virtual const TCHAR* GetName() const = 0;
 		virtual TArray<FString> GetSupportedCodecs() const = 0;
+		PRAGMA_DISABLE_DEPRECATION_WARNINGS
 		virtual TUniquePtr<FAudioEncoder> CreateEncoder(const FString& Codec) = 0;
+		PRAGMA_ENABLE_DEPRECATION_WARNINGS
 
 		static AVENCODER_API void RegisterFactory(FAudioEncoderFactory& Factory);
 		static AVENCODER_API void UnregisterFactory(FAudioEncoderFactory& Factory);
diff --git a/Engine/Source/Runtime/AVEncoder/Public/AudioTypes.h b/Engine/Source/Runtime/AVEncoder/Public/AudioTypes.h
--- a/Engine/Source/Runtime/AVEncoder/Public/AudioTypes.h
+++ b/Engine/Source/Runtime/AVEncoder/Public/AudioTypes.h
@@ -17,7 +17,7 @@ namespace LibRtc
 	using FAudioChannelSample = int16_t;
 
 	/** Describes the data format of the audio stream */
-	struct FAudioStreamFormat
+	struct UE_DEPRECATED(5.4, "AVEncoder has been deprecated. Please use the AVCodecs plugin family instead.") FAudioStreamFormat
 	{
 		/** The number of samples per seconds */
 		uint32_t SampleRate;
@@ -42,7 +42,7 @@ namespace LibRtc
 	 * One second - one frame * SampleRate
 	 * The number of samples per seconds - one sample byte * NumberBytesPerSample * NumberChannels * SampleRate
 	 */
-	struct FAudioSampleBuffer
+	struct UE_DEPRECATED(5.4, "AVEncoder has been deprecated. Please use the AVCodecs plugin family instead.") FAudioSampleBuffer
 	{
 		/** The sequence of audio samples */
 		TArrayView<const FAudioChannelSample> AudioSamples;
@@ -50,7 +50,7 @@ namespace LibRtc
 		FAudioStreamFormat AudioStreamFormat;
 	};
 
-	struct FAudioCodecConfig
+	struct UE_DEPRECATED(5.4, "AVEncoder has been deprecated. Please use the AVCodecs plugin family instead.") FAudioCodecConfig
 	{
 		int32_t SampleRate;
 		int32_t NumChannels;
@@ -78,7 +78,7 @@ namespace LibRtc
 	 */
 	using OnAudioDecodeCallback = TFunction<void(const FAudioByteBuffer& AudioBytes, const FAudioSampleBuffer&)>;
 
-	struct FAudioCodec
+	struct UE_DEPRECATED(5.4, "AVEncoder has been deprecated. Please use the AVCodecs plugin family instead.") FAudioCodec
 	{
 		FAudioCodecConfig AudioCodecConfig;
 		OnAudioEncodeCallback OnAudioEncodeCallback;
@@ -86,19 +86,19 @@ namespace LibRtc
 	};
 
 	/** Information which is used to initialize the incoming audio channel */
-	struct FIncomingAudioChannelInfo
+	struct UE_DEPRECATED(5.4, "AVEncoder has been deprecated. Please use the AVCodecs plugin family instead.") FIncomingAudioChannelInfo
 	{
 		// nothing
 	};
 
 	/** Information which is used to initialize the outgoing audio channel */
-	struct FOutgoingAudioChannelInfo
+	struct UE_DEPRECATED(5.4, "AVEncoder has been deprecated. Please use the AVCodecs plugin family instead.") FOutgoingAudioChannelInfo
 	{
 		/** Initial muted state of the audio channel */
 		bool bIsMuted;
 	};
 
-	enum class EAudioMuteReason
+	enum class UE_DEPRECATED(5.4, "AVEncoder has been deprecated. Please use the AVCodecs plugin family instead.") EAudioMuteReason
 	{
 		Unsupported,
 		Manual,
@@ -106,7 +106,7 @@ namespace LibRtc
 		AdminDisabled
 	};
 
-	struct FAudioMuteInfo
+	struct UE_DEPRECATED(5.4, "AVEncoder has been deprecated. Please use the AVCodecs plugin family instead.") FAudioMuteInfo
 	{
 		TOptional<EAudioMuteReason> Reason;
 	};
diff --git a/Engine/Source/Runtime/AVEncoder/Public/CodecPacket.h b/Engine/Source/Runtime/AVEncoder/Public/CodecPacket.h
--- a/Engine/Source/Runtime/AVEncoder/Public/CodecPacket.h
+++ b/Engine/Source/Runtime/AVEncoder/Public/CodecPacket.h
@@ -9,7 +9,7 @@
 
 namespace AVEncoder
 {
-	class FCodecPacket
+	class UE_DEPRECATED(5.4, "AVEncoder has been deprecated. Please use the AVCodecs plugin family instead.") FCodecPacket
 	{
 	public:
 		virtual ~FCodecPacket() = default;
@@ -19,7 +19,7 @@ namespace AVEncoder
 		/**
 		 * Encoding/Decoding latency
 		 */
-		struct FTimings
+		struct UE_DEPRECATED(5.4, "AVEncoder has been deprecated. Please use the AVCodecs plugin family instead.") FTimings
 		{
 			FTimespan StartTs;
 			FTimespan FinishTs;
diff --git a/Engine/Source/Runtime/AVEncoder/Public/Configuration.h b/Engine/Source/Runtime/AVEncoder/Public/Configuration.h
--- a/Engine/Source/Runtime/AVEncoder/Public/Configuration.h
+++ b/Engine/Source/Runtime/AVEncoder/Public/Configuration.h
@@ -12,9 +12,9 @@
 
 namespace LibRtc
 {
-    struct FConfiguration
+    struct UE_DEPRECATED(5.4, "AVEncoder has been deprecated. Please use the AVCodecs plugin family instead.") FConfiguration
     {
-        struct FData
+        struct UE_DEPRECATED(5.4, "AVEncoder has been deprecated. Please use the AVCodecs plugin family instead.") FData
         {
             static constexpr const TCHAR* kDefaultStunServerUrl = TEXT("stun:stun.l.google.com:19302");
 
diff --git a/Engine/Source/Runtime/AVEncoder/Public/IMFSampleWrapper.h b/Engine/Source/Runtime/AVEncoder/Public/IMFSampleWrapper.h
--- a/Engine/Source/Runtime/AVEncoder/Public/IMFSampleWrapper.h
+++ b/Engine/Source/Runtime/AVEncoder/Public/IMFSampleWrapper.h
@@ -15,15 +15,16 @@ namespace AVEncoder
 //
 // Wrapper for IMFSample, to make it easier to report errors
 //
-class FIMFSampleWrapper
+class UE_DEPRECATED(5.4, "AVEncoder has been deprecated. Please use the AVCodecs plugin family instead.") FIMFSampleWrapper
 {
 public:
-
+	PRAGMA_DISABLE_DEPRECATION_WARNINGS
 	FIMFSampleWrapper(EPacketType InMediaType = EPacketType::Invalid, IMFSample* InSample = nullptr)
 		: MediaType(InMediaType)
 		, Sample(InSample)
 	{
 	}
+	PRAGMA_ENABLE_DEPRECATION_WARNINGS
 
 	const IMFSample* GetSample() const
 	{
@@ -95,7 +96,9 @@ public:
 	AVENCODER_API FIMFSampleWrapper Clone() const;
 
 private:
+	PRAGMA_DISABLE_DEPRECATION_WARNINGS
 	EPacketType MediaType;
+	PRAGMA_ENABLE_DEPRECATION_WARNINGS
 	TRefCountPtr<IMFSample> Sample;
 };
 
diff --git a/Engine/Source/Runtime/AVEncoder/Public/MediaPacket.h b/Engine/Source/Runtime/AVEncoder/Public/MediaPacket.h
--- a/Engine/Source/Runtime/AVEncoder/Public/MediaPacket.h
+++ b/Engine/Source/Runtime/AVEncoder/Public/MediaPacket.h
@@ -5,7 +5,7 @@
 
 namespace AVEncoder
 {
-	struct FVideoConfig
+	struct UE_DEPRECATED(5.4, "AVEncoder has been deprecated. Please use the AVCodecs plugin family instead.") FVideoConfig
 	{
 		FString Codec;
 		uint32 Width;
@@ -14,7 +14,7 @@ namespace AVEncoder
 		uint32 Bitrate;
 	};
 
-	struct FAudioConfig
+	struct UE_DEPRECATED(5.4, "AVEncoder has been deprecated. Please use the AVCodecs plugin family instead.") FAudioConfig
 	{
 		FString Codec;
 		uint32 Samplerate;
@@ -22,16 +22,18 @@ namespace AVEncoder
 		uint32 Bitrate;
 	};
 
-	enum class EPacketType
+	enum class UE_DEPRECATED(5.4, "AVEncoder has been deprecated. Please use the AVCodecs plugin family instead.") EPacketType
 	{
 		Audio,
 		Video,
 		Invalid
 	};
 
-	struct FMediaPacket
+	struct UE_DEPRECATED(5.4, "AVEncoder has been deprecated. Please use the AVCodecs plugin family instead.") FMediaPacket
 	{
+		PRAGMA_DISABLE_DEPRECATION_WARNINGS
 		EPacketType Type;
+		PRAGMA_ENABLE_DEPRECATION_WARNINGS
 		FTimespan Timestamp;
 		FTimespan Duration;
 		TArray<uint8> Data;
@@ -52,13 +54,19 @@ namespace AVEncoder
 			} Audio;
 		};
 
+		PRAGMA_DISABLE_DEPRECATION_WARNINGS
 		explicit FMediaPacket(EPacketType TypeIn)
+		PRAGMA_ENABLE_DEPRECATION_WARNINGS
 		{
 			Type = TypeIn;
+			PRAGMA_DISABLE_DEPRECATION_WARNINGS
 			if (Type == EPacketType::Audio)
+			PRAGMA_ENABLE_DEPRECATION_WARNINGS
 			{
 			}
+			PRAGMA_DISABLE_DEPRECATION_WARNINGS
 			else if (Type == EPacketType::Video)
+			PRAGMA_ENABLE_DEPRECATION_WARNINGS
 			{
 				FMemory::Memzero(Video);
 			}
@@ -66,7 +74,9 @@ namespace AVEncoder
 
 		bool IsVideoKeyFrame() const
 		{
+			PRAGMA_DISABLE_DEPRECATION_WARNINGS
 			return (Type == EPacketType::Video && Video.bKeyFrame) ? true : false;
+			PRAGMA_ENABLE_DEPRECATION_WARNINGS
 		}
 	};
 }
diff --git a/Engine/Source/Runtime/AVEncoder/Public/VideoCommon.h b/Engine/Source/Runtime/AVEncoder/Public/VideoCommon.h
--- a/Engine/Source/Runtime/AVEncoder/Public/VideoCommon.h
+++ b/Engine/Source/Runtime/AVEncoder/Public/VideoCommon.h
@@ -47,7 +47,7 @@ namespace AVEncoder
 {
 	const int64 TimeStampNone = 0x7fffffffll;
 
-	enum class EVideoFrameFormat
+	enum class UE_DEPRECATED(5.4, "AVEncoder has been deprecated. Please use the AVCodecs plugin family instead.") EVideoFrameFormat
 	{
 		Undefined,				// (not-yet) defined format
 		YUV420P,				// Planar YUV420 format in CPU memory
@@ -57,7 +57,7 @@ namespace AVEncoder
 		VULKAN_R8G8B8A8_UNORM,
 	};
 
-	enum class EH264Profile
+	enum class UE_DEPRECATED(5.4, "AVEncoder has been deprecated. Please use the AVCodecs plugin family instead.") EH264Profile
 	{
 		UNKNOWN,
 		CONSTRAINED_BASELINE,
@@ -67,6 +67,7 @@ namespace AVEncoder
 		HIGH,
 	};
 
+	PRAGMA_DISABLE_DEPRECATION_WARNINGS
 	inline FString ToString(EVideoFrameFormat Format)
 	{
 		switch (Format)
@@ -86,8 +87,9 @@ namespace AVEncoder
 			return FString("EVideoFrameFormat::Undefined");
 		}
 	}
+	PRAGMA_ENABLE_DEPRECATION_WARNINGS
 
-	enum class ECodecType
+	enum class UE_DEPRECATED(5.4, "AVEncoder has been deprecated. Please use the AVCodecs plugin family instead.") ECodecType
 	{
 		Undefined,
 		H264,
@@ -102,13 +104,16 @@ namespace AVEncoder
 	const uint32 H264Profile_ConstrainedHigh = 1 << 3;
 	const uint32 H264Profile_High = 1 << 4;
 
-	struct FVideoEncoderInfo
+	struct UE_DEPRECATED(5.4, "AVEncoder has been deprecated. Please use the AVCodecs plugin family instead.") FVideoEncoderInfo
 	{
 		uint32						ID = 0;
-		ECodecType					CodecType = ECodecType::Undefined;
 		uint32						MaxWidth = 0;
 		uint32						MaxHeight = 0;
+		PRAGMA_DISABLE_DEPRECATION_WARNINGS
+		ECodecType					CodecType = ECodecType::Undefined;
 		TArray<EVideoFrameFormat>	SupportedInputFormats;
+		PRAGMA_ENABLE_DEPRECATION_WARNINGS
+
 		struct
 		{
 			uint32					SupportedProfiles = 0;
@@ -118,12 +123,15 @@ namespace AVEncoder
 	};
 
 
-	struct FVideoDecoderInfo
+	struct UE_DEPRECATED(5.4, "AVEncoder has been deprecated. Please use the AVCodecs plugin family instead.") FVideoDecoderInfo
 	{
 		uint32						ID = 0;
-		ECodecType					CodecType = ECodecType::Undefined;
 		uint32						MaxWidth = 0;
 		uint32						MaxHeight = 0;
+
+		PRAGMA_DISABLE_DEPRECATION_WARNINGS
+		ECodecType					CodecType = ECodecType::Undefined;
+		PRAGMA_ENABLE_DEPRECATION_WARNINGS
 	};
 
 #if PLATFORM_WINDOWS
diff --git a/Engine/Source/Runtime/AVEncoder/Public/VideoDecoder.h b/Engine/Source/Runtime/AVEncoder/Public/VideoDecoder.h
--- a/Engine/Source/Runtime/AVEncoder/Public/VideoDecoder.h
+++ b/Engine/Source/Runtime/AVEncoder/Public/VideoDecoder.h
@@ -15,13 +15,13 @@ namespace AVEncoder
 	class FVideoDecoderInput;
 	class FVideoDecoderOutput;
 
-	class FVideoDecoder
+	class UE_DEPRECATED(5.4, "AVEncoder has been deprecated. Please use the AVCodecs plugin family instead.") FVideoDecoder
 	{
 	public:
 		using CreateDecoderAllocationInterfaceCallback = TFunction<void(void* /*InOptions*/, void** /*InOutParamResult*/)>;
 		using ReleaseDecoderAllocationInterfaceCallback = TFunction<void(void* /*InOptions*/, void** /*InOutParamResult*/)>;
 
-		struct FInit
+		struct UE_DEPRECATED(5.4, "AVEncoder has been deprecated. Please use the AVCodecs plugin family instead.") FInit
 		{
 			CreateDecoderAllocationInterfaceCallback	CreateDecoderAllocationInterface;
 			ReleaseDecoderAllocationInterfaceCallback	ReleaseDecoderAllocationInterface;
@@ -35,15 +35,18 @@ namespace AVEncoder
 		// Shuts the decoder down AND destroys it. Do not store a pointer to this decoder in a smart pointer!
 		virtual void Shutdown() = 0;
 
-		enum class EDecodeResult
+		enum class UE_DEPRECATED(5.4, "AVEncoder has been deprecated. Please use the AVCodecs plugin family instead.") EDecodeResult
 		{
 			Success,
 			Failure
 //			TryAgainLater
 		};
+
+		PRAGMA_DISABLE_DEPRECATION_WARNINGS
 		virtual EDecodeResult Decode(const FVideoDecoderInput* InInput) = 0;
 
 		using OnDecodedFrameCallback = TFunction<void(const FVideoDecoderOutput* /*InDecodedFrame*/)>;
+		PRAGMA_ENABLE_DEPRECATION_WARNINGS
 
 		virtual void SetOnDecodedFrame(OnDecodedFrameCallback InCallback) { OnDecodedFrame = MoveTemp(InCallback); }
 		virtual void ClearOnDecodedFrame() { OnDecodedFrame = nullptr; }
@@ -55,8 +58,10 @@ namespace AVEncoder
 	protected:
 		bool CreateDecoderAllocationInterface();
 		void ReleaseDecoderAllocationInterface();
-		EFrameBufferAllocReturn AllocateOutputFrameBuffer(FVideoDecoderAllocFrameBufferResult* OutBuffer, const FVideoDecoderAllocFrameBufferParams* InAllocParams);
 		void* GetAllocationInterfaceMethods();
+		PRAGMA_DISABLE_DEPRECATION_WARNINGS
+		EFrameBufferAllocReturn AllocateOutputFrameBuffer(FVideoDecoderAllocFrameBufferResult* OutBuffer, const FVideoDecoderAllocFrameBufferParams* InAllocParams);
+		PRAGMA_ENABLE_DEPRECATION_WARNINGS
 
 		struct FPlatformDecoderAllocInterface;
 
diff --git a/Engine/Source/Runtime/AVEncoder/Public/VideoDecoderAllocationTypes.h b/Engine/Source/Runtime/AVEncoder/Public/VideoDecoderAllocationTypes.h
--- a/Engine/Source/Runtime/AVEncoder/Public/VideoDecoderAllocationTypes.h
+++ b/Engine/Source/Runtime/AVEncoder/Public/VideoDecoderAllocationTypes.h
@@ -9,7 +9,7 @@
 /**
  * Type of video frame buffer to allocate.
  */
-enum class EFrameBufferType
+enum class UE_DEPRECATED(5.4, "AVEncoder has been deprecated. Please use the AVCodecs plugin family instead.") EFrameBufferType
 {
 	CODEC_RawBuffer = 0,
 	CODEC_TextureHandle = 1,
@@ -19,10 +19,12 @@ enum class EFrameBufferType
 /**
  * Video frame buffer allocation parameters.
  */
-struct FVideoDecoderAllocFrameBufferParams
+struct UE_DEPRECATED(5.4, "AVEncoder has been deprecated. Please use the AVCodecs plugin family instead.") FVideoDecoderAllocFrameBufferParams
 {
 	/** Type of output frame buffer to allocate */
+	PRAGMA_DISABLE_DEPRECATION_WARNINGS
 	EFrameBufferType FrameBufferType;
+	PRAGMA_ENABLE_DEPRECATION_WARNINGS
 	/**
 	 * Number of bytes to allocate with the specified alignment.
 	 * For a raw buffer the allocation size is required as it may include oversize for padding.
@@ -44,7 +46,7 @@ struct FVideoDecoderAllocFrameBufferParams
 /**
  * Return value for the video decoder frame buffer allocation callback.
  */
-enum class EFrameBufferAllocReturn
+enum class UE_DEPRECATED(5.4, "AVEncoder has been deprecated. Please use the AVCodecs plugin family instead.") EFrameBufferAllocReturn
 {
 	/** Buffer successfully allocated. */
 	CODEC_Success = 0,
@@ -67,7 +69,7 @@ DECLARE_DELEGATE_TwoParams(FReleaseFrameBuffer, void* /* , This */, void* /*Buff
 /**
  * Output buffer plane description.
  */
-struct FFrameBufferOutPlaneDesc
+struct UE_DEPRECATED(5.4, "AVEncoder has been deprecated. Please use the AVCodecs plugin family instead.") FFrameBufferOutPlaneDesc
 {
 	/** Width of the allocated buffer plane in pixels, including necessary padding. */
 	int32_t Width;
@@ -86,7 +88,7 @@ struct FFrameBufferOutPlaneDesc
 /**
  * Result structure to be filled in by the application.
  */
-struct FVideoDecoderAllocFrameBufferResult
+struct UE_DEPRECATED(5.4, "AVEncoder has been deprecated. Please use the AVCodecs plugin family instead.") FVideoDecoderAllocFrameBufferResult
 {
 	/** The buffer the application has allocated. */
 	void* AllocatedBuffer;
@@ -97,7 +99,9 @@ struct FVideoDecoderAllocFrameBufferResult
 	 *  to allocate). Supports at most 4 planes (R,G,B,A). Usually 2 (Y, UV) or 3 (Y,U,V or R,G,B). */
 	int32_t AllocatedPlanesNum;
 	int32_t AllocatedPlaneLayout;	/* reserved */
+	PRAGMA_DISABLE_DEPRECATION_WARNINGS
 	FFrameBufferOutPlaneDesc AllocatedPlaneDesc[4];
+	PRAGMA_ENABLE_DEPRECATION_WARNINGS
 
 	/** Callback within the application to retain the allocated buffer. */
 	FRetainFrameBuffer RetainCallback;
@@ -112,7 +116,9 @@ struct FVideoDecoderAllocFrameBufferResult
  * Video frame buffer allocation callback within the application.
  * We call this to get a new output frame buffer to decode into.
  */
+PRAGMA_DISABLE_DEPRECATION_WARNINGS
 DECLARE_DELEGATE_RetVal_ThreeParams(EFrameBufferAllocReturn, FAllocFrameBuffer, void* /*This*/, const FVideoDecoderAllocFrameBufferParams* /*InAllocParams*/, FVideoDecoderAllocFrameBufferResult* /*OutBuffer*/);
+PRAGMA_ENABLE_DEPRECATION_WARNINGS
 
 
 
@@ -121,7 +127,7 @@ DECLARE_DELEGATE_RetVal_ThreeParams(EFrameBufferAllocReturn, FAllocFrameBuffer,
 
 DECLARE_DELEGATE_RetVal_ThreeParams(int32_t, FGetD3DDevice, void* /* , This */, void** /* , OutD3DDevice */, int32_t* /* , OutD3DVersion */);
 
-struct FVideoDecoderMethodsWindows
+struct UE_DEPRECATED(5.4, "AVEncoder has been deprecated. Please use the AVCodecs plugin family instead.") FVideoDecoderMethodsWindows
 {
 	/** Magic cookie to check for if this is indeed the expected structure. */
 	uint32_t MagicCookie;	// 'WinX' 0x57696e58
@@ -139,7 +145,7 @@ DECLARE_DELEGATE_RetVal_FiveParams(int32_t, FMemAlloc, void* /*, This*/, void**
 DECLARE_DELEGATE_RetVal_ThreeParams(int32_t, FMemFree, void* /* , This */, void** /* , InAddress */, int32_t /* , MemType */);
 DECLARE_DELEGATE_RetVal_ThreeParams(int32_t, FConfigQuery, void* /* , This */, void** /* , OutInfo */, void** /* , InInfo */);
 
-struct FVideoDecoderMethods
+struct UE_DEPRECATED(5.4, "AVEncoder has been deprecated. Please use the AVCodecs plugin family instead.") FVideoDecoderMethods
 {
 	/** Magic cookie to check for if this is indeed the expected structure. */
 	uint32_t MagicCookie;
diff --git a/Engine/Source/Runtime/AVEncoder/Public/VideoDecoderFactory.h b/Engine/Source/Runtime/AVEncoder/Public/VideoDecoderFactory.h
--- a/Engine/Source/Runtime/AVEncoder/Public/VideoDecoderFactory.h
+++ b/Engine/Source/Runtime/AVEncoder/Public/VideoDecoderFactory.h
@@ -15,7 +15,7 @@ namespace AVEncoder
 {
 	struct FVideoDecoderInfo;
 
-	class FVideoDecoderFactory
+	class UE_DEPRECATED(5.4, "AVEncoder has been deprecated. Please use the AVCodecs plugin family instead.") FVideoDecoderFactory
 	{
 	public:
 		static AVENCODER_API FVideoDecoderFactory& Get();
@@ -24,6 +24,7 @@ namespace AVEncoder
 		// --- decoder registry
 
 		// the callback type used to create a registered decoder
+		PRAGMA_DISABLE_DEPRECATION_WARNINGS
 		using CreateDecoderCallback = TFunction<FVideoDecoder*()>;
 
 		// register a decoder so that it can be iterated and created
@@ -36,6 +37,7 @@ namespace AVEncoder
 		// --- decoder creation
 		// create a decoder instance. If not nullptr it must be destroyed through invoking its Shutdown() method!
 		AVENCODER_API FVideoDecoder* Create(uint32 InID, const FVideoDecoder::FInit& InInit);
+		PRAGMA_ENABLE_DEPRECATION_WARNINGS
 
 	private:
 		AVENCODER_API FVideoDecoderFactory();
@@ -48,8 +50,11 @@ namespace AVEncoder
 		static AVENCODER_API FThreadSafeCounter		NextID;
 		static AVENCODER_API bool						bDebugDontRegisterDefaultCodecs;
 
-		TArray<FVideoDecoderInfo>		AvailableDecoders;
 		TArray<CreateDecoderCallback>	CreateDecoders;
+
+		PRAGMA_DISABLE_DEPRECATION_WARNINGS
+		TArray<FVideoDecoderInfo>		AvailableDecoders;
+		PRAGMA_ENABLE_DEPRECATION_WARNINGS
 	};
 
 } /* namespace AVEncoder */
diff --git a/Engine/Source/Runtime/AVEncoder/Public/VideoDecoderInput.h b/Engine/Source/Runtime/AVEncoder/Public/VideoDecoderInput.h
--- a/Engine/Source/Runtime/AVEncoder/Public/VideoDecoderInput.h
+++ b/Engine/Source/Runtime/AVEncoder/Public/VideoDecoderInput.h
@@ -14,11 +14,11 @@ struct FVideoDecoderAllocFrameBufferResult;
 namespace AVEncoder
 {
 
-	class FVideoDecoderInput
+	class UE_DEPRECATED(5.4, "AVEncoder has been deprecated. Please use the AVCodecs plugin family instead.") FVideoDecoderInput
 	{
 	public:
 
-		struct FInputData
+		struct UE_DEPRECATED(5.4, "AVEncoder has been deprecated. Please use the AVCodecs plugin family instead.") FInputData
 		{
 			const void* EncodedData = nullptr;
 			int64		PTS = 0;
@@ -54,7 +54,7 @@ namespace AVEncoder
 
 
 
-	class FVideoDecoderOutput
+	class UE_DEPRECATED(5.4, "AVEncoder has been deprecated. Please use the AVCodecs plugin family instead.") FVideoDecoderOutput
 	{
 	public:
 		virtual int32 AddRef() = 0;
@@ -63,7 +63,9 @@ namespace AVEncoder
 		virtual int32 GetHeight() const = 0;
 		virtual int64 GetPTS() const = 0;
 
+		PRAGMA_DISABLE_DEPRECATION_WARNINGS
 		virtual const FVideoDecoderAllocFrameBufferResult* GetAllocatedBuffer() const = 0;
+		PRAGMA_ENABLE_DEPRECATION_WARNINGS
 
 		virtual int32 GetCropLeft() const = 0;
 		virtual int32 GetCropRight() const = 0;
diff --git a/Engine/Source/Runtime/AVEncoder/Public/VideoEncoder.h b/Engine/Source/Runtime/AVEncoder/Public/VideoEncoder.h
--- a/Engine/Source/Runtime/AVEncoder/Public/VideoEncoder.h
+++ b/Engine/Source/Runtime/AVEncoder/Public/VideoEncoder.h
@@ -19,14 +19,14 @@ namespace AVEncoder
 	class FVideoEncoderInput;
 	class FVideoEncoderInputFrame;
 
-    class FVideoEncoder
+    class UE_DEPRECATED(5.4, "AVEncoder has been deprecated. Please use the AVCodecs plugin family instead.") FVideoEncoder
     {
     public:
         enum class RateControlMode { UNKNOWN, CONSTQP, VBR, CBR };
         enum class MultipassMode { UNKNOWN, DISABLED, QUARTER, FULL };
 		enum class H264Profile { AUTO, CONSTRAINED_BASELINE, BASELINE, MAIN, CONSTRAINED_HIGH, HIGH, HIGH444, STEREO, SVC_TEMPORAL_SCALABILITY, PROGRESSIVE_HIGH };
 
-        struct FLayerConfig
+        struct UE_DEPRECATED(5.4, "AVEncoder has been deprecated. Please use the AVCodecs plugin family instead.") FLayerConfig
         {
             uint32			Width = 0;
             uint32			Height = 0;
@@ -63,7 +63,9 @@ namespace AVEncoder
 
         AVENCODER_API virtual ~FVideoEncoder();
 
+        PRAGMA_DISABLE_DEPRECATION_WARNINGS
 		virtual bool Setup(TSharedRef<FVideoEncoderInput> input, FLayerConfig const& config) { return false; }
+        PRAGMA_ENABLE_DEPRECATION_WARNINGS
 		virtual void Shutdown() {}
 
         AVENCODER_API virtual bool AddLayer(FLayerConfig const& config);
@@ -73,10 +75,12 @@ namespace AVEncoder
 		AVENCODER_API FLayerConfig GetLayerConfig(uint32 layerIdx) const;
         AVENCODER_API void UpdateLayerConfig(uint32 layerIdx, FLayerConfig const& config);
 
+        PRAGMA_DISABLE_DEPRECATION_WARNINGS
         using OnFrameEncodedCallback = TFunction<void(const TSharedPtr<FVideoEncoderInputFrame> /* InCompletedFrame */)>;
         using OnEncodedPacketCallback = TFunction<void(uint32 /* LayerIndex */, const TSharedPtr<FVideoEncoderInputFrame> /* Frame */, const FCodecPacket& /* Packet */)>;
+        PRAGMA_ENABLE_DEPRECATION_WARNINGS
 
-        struct FEncodeOptions
+        struct UE_DEPRECATED(5.4, "AVEncoder has been deprecated. Please use the AVCodecs plugin family instead.") FEncodeOptions
         {
             bool					bForceKeyFrame = false;
             OnFrameEncodedCallback	OnFrameEncoded;
@@ -85,12 +89,14 @@ namespace AVEncoder
         void SetOnEncodedPacket(OnEncodedPacketCallback callback) { OnEncodedPacket = MoveTemp(callback); }
         void ClearOnEncodedPacket() { OnEncodedPacket = nullptr; }
 
+        PRAGMA_DISABLE_DEPRECATION_WARNINGS
 		virtual void Encode(const TSharedPtr<FVideoEncoderInputFrame> frame, FEncodeOptions const& options) {}
+        PRAGMA_ENABLE_DEPRECATION_WARNINGS
 
     protected:
         FVideoEncoder() = default;
 
-        class FLayer
+        class UE_DEPRECATED(5.4, "AVEncoder has been deprecated. Please use the AVCodecs plugin family instead.") FLayer
         {
         public:
             explicit FLayer(FLayerConfig const& layerConfig)
diff --git a/Engine/Source/Runtime/AVEncoder/Public/VideoEncoderFactory.h b/Engine/Source/Runtime/AVEncoder/Public/VideoEncoderFactory.h
--- a/Engine/Source/Runtime/AVEncoder/Public/VideoEncoderFactory.h
+++ b/Engine/Source/Runtime/AVEncoder/Public/VideoEncoderFactory.h
@@ -18,7 +18,7 @@ namespace AVEncoder
 {
 	class FVideoEncoderInput;
 
-	class FVideoEncoderFactory
+	class UE_DEPRECATED(5.4, "AVEncoder has been deprecated. Please use the AVCodecs plugin family instead.") FVideoEncoderFactory
 	{
 	public:
 		static AVENCODER_API FVideoEncoderFactory& Get();
@@ -30,6 +30,7 @@ namespace AVEncoder
 		// --- encoder registry
 
 		// the callback type used to create a registered encoder
+		PRAGMA_DISABLE_DEPRECATION_WARNINGS
 		using CreateEncoderCallback = TFunction<TUniquePtr<FVideoEncoder>()>;
 
 		// register an encoder so that it can be iterated and created
@@ -44,7 +45,7 @@ namespace AVEncoder
 
 		AVENCODER_API TUniquePtr<FVideoEncoder> Create(uint32 InID, const FVideoEncoder::FLayerConfig& config);
 		AVENCODER_API TUniquePtr<FVideoEncoder> Create(uint32 InID, TSharedPtr<FVideoEncoderInput> InInput, const FVideoEncoder::FLayerConfig& config);
-
+		PRAGMA_ENABLE_DEPRECATION_WARNINGS
 	private:
 		FVideoEncoderFactory() = default;
 		~FVideoEncoderFactory() = default;
@@ -58,9 +59,11 @@ namespace AVEncoder
 		static AVENCODER_API FThreadSafeCounter		NextID;
 		bool							bDebugDontRegisterDefaultCodecs = false;
 		FThreadSafeBool					bWasSetup;
-		TArray<FVideoEncoderInfo>		AvailableEncoders;
 		TArray<CreateEncoderCallback>	CreateEncoders;
 
+		PRAGMA_DISABLE_DEPRECATION_WARNINGS
+		TArray<FVideoEncoderInfo>		AvailableEncoders;
+		PRAGMA_ENABLE_DEPRECATION_WARNINGS
 	};
     
 } /* namespace AVEncoder */
diff --git a/Engine/Source/Runtime/AVEncoder/Public/VideoEncoderInput.h b/Engine/Source/Runtime/AVEncoder/Public/VideoEncoderInput.h
--- a/Engine/Source/Runtime/AVEncoder/Public/VideoEncoderInput.h
+++ b/Engine/Source/Runtime/AVEncoder/Public/VideoEncoderInput.h
@@ -34,7 +34,7 @@ namespace AVEncoder
 	class FVideoEncoderInputFrame;
 
 #if PLATFORM_DESKTOP && !PLATFORM_APPLE
-	struct FVulkanDataStruct
+	struct UE_DEPRECATED(5.4, "AVEncoder has been deprecated. Please use the AVCodecs plugin family instead.") FVulkanDataStruct
 	{
 		VkInstance VulkanInstance;
 		VkPhysicalDevice VulkanPhysicalDevice;
@@ -42,7 +42,7 @@ namespace AVEncoder
 	};
 #endif
 
-	class FVideoEncoderInput
+	class UE_DEPRECATED(5.4, "AVEncoder has been deprecated. Please use the AVCodecs plugin family instead.") FVideoEncoderInput
 	{
 	public:
 		// --- construct video encoder input based on expected input frame format
@@ -64,12 +64,16 @@ namespace AVEncoder
 		// --- properties
 		AVENCODER_API virtual void SetMaxNumBuffers(uint32 InMaxNumBuffers);
 
+		PRAGMA_DISABLE_DEPRECATION_WARNINGS
 		EVideoFrameFormat GetFrameFormat() const { return FrameFormat; }
+		PRAGMA_ENABLE_DEPRECATION_WARNINGS
 
 		// --- available encoders
 
 		// get a list of supported video encoders
+		PRAGMA_DISABLE_DEPRECATION_WARNINGS
 		virtual const TArray<FVideoEncoderInfo>& GetAvailableEncoders() = 0;
+		PRAGMA_ENABLE_DEPRECATION_WARNINGS
 
 		// --- create encoders
 
@@ -112,7 +116,10 @@ namespace AVEncoder
 		FVideoEncoderInput(const FVideoEncoderInput&) = delete;
 		FVideoEncoderInput& operator=(const FVideoEncoderInput&) = delete;
 
+		PRAGMA_DISABLE_DEPRECATION_WARNINGS
 		EVideoFrameFormat				FrameFormat = EVideoFrameFormat::Undefined;
+		PRAGMA_ENABLE_DEPRECATION_WARNINGS
+
 		uint32 MaxNumBuffers = 3;
 		uint32 NumBuffers = 0;
 
@@ -123,7 +130,7 @@ namespace AVEncoder
 
 
 	// TODO this should go elsewhere and be made cross platform
-	class FVideoEncoderInputFrame
+	class UE_DEPRECATED(5.4, "AVEncoder has been deprecated. Please use the AVCodecs plugin family instead.") FVideoEncoderInputFrame
 	{
 	public:
 		// Obtain (increase reference count) of this input frame
@@ -147,7 +154,9 @@ namespace AVEncoder
 		int64 GetTimestampRTP() const { return TimestampRTP; }
 
 		// current format of frame
+		PRAGMA_DISABLE_DEPRECATION_WARNINGS
 		EVideoFrameFormat GetFormat() const { return Format; }
+		PRAGMA_ENABLE_DEPRECATION_WARNINGS
 		// width of frame buffer
 		void SetWidth(uint32 InWidth) { Width = InWidth; }
 		uint32 GetWidth() const { return Width; }
@@ -159,7 +168,7 @@ namespace AVEncoder
 
 		// --- YUV420P
 
-		struct FYUV420P
+		struct UE_DEPRECATED(5.4, "AVEncoder has been deprecated. Please use the AVCodecs plugin family instead.") FYUV420P
 		{
 			const uint8*		Data[3] = { nullptr, nullptr, nullptr };
 			uint32				StrideY = 0;
@@ -180,7 +189,7 @@ namespace AVEncoder
 #if PLATFORM_WINDOWS
 		// --- D3D11
 
-		struct FD3D11
+		struct UE_DEPRECATED(5.4, "AVEncoder has been deprecated. Please use the AVCodecs plugin family instead.") FD3D11
 		{
 			ID3D11Texture2D*	Texture = nullptr;
 			ID3D11Device*		EncoderDevice = nullptr;
@@ -198,7 +207,7 @@ namespace AVEncoder
 
 		// --- D3D12
 
-		struct FD3D12
+		struct UE_DEPRECATED(5.4, "AVEncoder has been deprecated. Please use the AVCodecs plugin family instead.") FD3D12
 		{
 			ID3D12Resource*		Texture = nullptr;
 			ID3D12Device*		EncoderDevice = nullptr;
@@ -215,7 +224,7 @@ namespace AVEncoder
 
 #endif // PLATFORM_WINDOWS
 
-		enum class EUnderlyingRHI
+		enum class UE_DEPRECATED(5.4, "AVEncoder has been deprecated. Please use the AVCodecs plugin family instead.") EUnderlyingRHI
 		{
 			Undefined,
 			D3D11,
@@ -224,12 +233,15 @@ namespace AVEncoder
 		};
 
 		// --- CUDA
-		struct FCUDA
+		struct UE_DEPRECATED(5.4, "AVEncoder has been deprecated. Please use the AVCodecs plugin family instead.") FCUDA
 		{
 			CUarray			EncoderTexture = nullptr;
 			CUcontext   	EncoderDevice = nullptr;
-			EUnderlyingRHI	UnderlyingRHI = EUnderlyingRHI::Undefined;
 			void*			SharedHandle = nullptr;
+
+			PRAGMA_DISABLE_DEPRECATION_WARNINGS
+			EUnderlyingRHI	UnderlyingRHI = EUnderlyingRHI::Undefined;
+			PRAGMA_ENABLE_DEPRECATION_WARNINGS
 		};
 
 		const FCUDA& GetCUDA() const { return CUDA; }
@@ -242,7 +254,7 @@ namespace AVEncoder
 
 #if PLATFORM_DESKTOP && !PLATFORM_APPLE
 		// --- Vulkan
-		struct FVulkan
+		struct UE_DEPRECATED(5.4, "AVEncoder has been deprecated. Please use the AVCodecs plugin family instead.") FVulkan
 		{
 			VkImage				EncoderTexture = VK_NULL_HANDLE;
 			VkDeviceMemory		EncoderDeviceMemory;
@@ -273,12 +285,15 @@ namespace AVEncoder
 		int64									TimestampUs;
 		int64									TimestampRTP;
 		mutable FThreadSafeCounter				NumReferences;
-		EVideoFrameFormat						Format;
 		uint32									Width;
 		uint32									Height;
 		FYUV420P								YUV420P;
 		bool									bFreeYUV420PData;
 
+		PRAGMA_DISABLE_DEPRECATION_WARNINGS
+		EVideoFrameFormat						Format;
+		PRAGMA_ENABLE_DEPRECATION_WARNINGS
+
 #if PLATFORM_WINDOWS
 		FD3D11									D3D11;
 		FReleaseD3D11TextureCallback			OnReleaseD3D11Texture;
diff --git a/Engine/Source/Runtime/GameplayMediaEncoder/Private/GameplayMediaEncoder.cpp b/Engine/Source/Runtime/GameplayMediaEncoder/Private/GameplayMediaEncoder.cpp
--- a/Engine/Source/Runtime/GameplayMediaEncoder/Private/GameplayMediaEncoder.cpp
+++ b/Engine/Source/Runtime/GameplayMediaEncoder/Private/GameplayMediaEncoder.cpp
@@ -49,14 +49,15 @@ const uint32 MaxVideoFPS = 60;
 const uint32 MaxWidth = 1920;
 const uint32 MaxHeight = 1080;
 
-FAutoConsoleCommand GameplayMediaEncoderInitialize(TEXT("GameplayMediaEncoder.Initialize"), TEXT("Constructs the audio/video encoding objects. Does not start encoding"),
-                                                   FConsoleCommandDelegate::CreateStatic(&FGameplayMediaEncoder::InitializeCmd));
+PRAGMA_DISABLE_DEPRECATION_WARNINGS
+FAutoConsoleCommand GameplayMediaEncoderInitialize(TEXT("GameplayMediaEncoder.Initialize"), TEXT("Constructs the audio/video encoding objects. Does not start encoding"), FConsoleCommandDelegate::CreateStatic(&FGameplayMediaEncoder::InitializeCmd));
 
 FAutoConsoleCommand GameplayMediaEncoderStart(TEXT("GameplayMediaEncoder.Start"), TEXT("Starts encoding"), FConsoleCommandDelegate::CreateStatic(&FGameplayMediaEncoder::StartCmd));
 
 FAutoConsoleCommand GameplayMediaEncoderStop(TEXT("GameplayMediaEncoder.Stop"), TEXT("Stops encoding"), FConsoleCommandDelegate::CreateStatic(&FGameplayMediaEncoder::StopCmd));
 
 FAutoConsoleCommand GameplayMediaEncoderShutdown(TEXT("GameplayMediaEncoder.Shutdown"), TEXT("Releases all systems."), FConsoleCommandDelegate::CreateStatic(&FGameplayMediaEncoder::ShutdownCmd));
+PRAGMA_ENABLE_DEPRECATION_WARNINGS
 
 //////////////////////////////////////////////////////////////////////////
 //
@@ -64,21 +65,32 @@ FAutoConsoleCommand GameplayMediaEncoderShutdown(TEXT("GameplayMediaEncoder.Shut
 //
 //////////////////////////////////////////////////////////////////////////
 
+PRAGMA_DISABLE_DEPRECATION_WARNINGS
 TSharedPtr<FGameplayMediaEncoder, ESPMode::ThreadSafe> FGameplayMediaEncoder::Singleton = { };
+PRAGMA_ENABLE_DEPRECATION_WARNINGS
 
+PRAGMA_DISABLE_DEPRECATION_WARNINGS
 FGameplayMediaEncoder* FGameplayMediaEncoder::Get()
+PRAGMA_ENABLE_DEPRECATION_WARNINGS
 {
 	// Constructed and captured as a thread safe shared pointer as this is required by the audio submix listener interface
 	if(!Singleton.IsValid())
 	{
+		PRAGMA_DISABLE_DEPRECATION_WARNINGS
 		Singleton = MakeShared<FGameplayMediaEncoder>();
+		PRAGMA_ENABLE_DEPRECATION_WARNINGS
 	}
+	
 	return Singleton.Get();
 }
 
+PRAGMA_DISABLE_DEPRECATION_WARNINGS
 FGameplayMediaEncoder::FGameplayMediaEncoder() {}
+PRAGMA_ENABLE_DEPRECATION_WARNINGS
 
+PRAGMA_DISABLE_DEPRECATION_WARNINGS
 FGameplayMediaEncoder::~FGameplayMediaEncoder() { Shutdown(); }
+PRAGMA_ENABLE_DEPRECATION_WARNINGS
 
 bool FGameplayMediaEncoder::RegisterListener(IGameplayMediaEncoderListener* Listener)
 {
@@ -137,20 +149,25 @@ bool FGameplayMediaEncoder::Initialize()
 	//
 	// Audio
 	//
+	PRAGMA_DISABLE_DEPRECATION_WARNINGS
 	AVEncoder::FAudioEncoderFactory* AudioEncoderFactory = AVEncoder::FAudioEncoderFactory::FindFactory("aac");
+	PRAGMA_ENABLE_DEPRECATION_WARNINGS
 	if(!AudioEncoderFactory)
 	{
 		UE_LOG(GameplayMediaEncoder, Error, TEXT("No audio encoder for aac found"));
 		return false;
 	}
 
+	PRAGMA_DISABLE_DEPRECATION_WARNINGS
 	AudioEncoder = AudioEncoderFactory->CreateEncoder("aac");
+	PRAGMA_ENABLE_DEPRECATION_WARNINGS
 	if(!AudioEncoder)
 	{
 		UE_LOG(GameplayMediaEncoder, Error, TEXT("Could not create audio encoder"));
 		return false;
 	}
 
+	PRAGMA_DISABLE_DEPRECATION_WARNINGS
 	AVEncoder::FAudioConfig AudioConfig;
 	AudioConfig.Samplerate = HardcodedAudioSamplerate;
 	AudioConfig.NumChannels = HardcodedAudioNumChannels;
@@ -162,13 +179,14 @@ bool FGameplayMediaEncoder::Initialize()
 	}
 
 	AudioEncoder->RegisterListener(*this);
+	PRAGMA_ENABLE_DEPRECATION_WARNINGS
 
 	MemoryCheckpoint("Audio encoder initialized");
 
 	//
 	// Video
 	//
-
+	PRAGMA_DISABLE_DEPRECATION_WARNINGS
 	VideoConfig.Codec = "h264";
 	VideoConfig.Height = VideoConfig.Width = VideoConfig.Framerate = VideoConfig.Bitrate = 0;
 	FParse::Value(FCommandLine::Get(), TEXT("GameplayMediaEncoder.ResY="), VideoConfig.Height);
@@ -183,6 +201,7 @@ bool FGameplayMediaEncoder::Initialize()
 		VideoConfig.Width = 1920;
 		VideoConfig.Height = 1080;
 	}
+	PRAGMA_ENABLE_DEPRECATION_WARNINGS
 	else
 	{
 		UE_LOG(GameplayMediaEncoder, Fatal, TEXT("GameplayMediaEncoder.ResY can only have a value of 720 or 1080"));
@@ -190,22 +209,29 @@ bool FGameplayMediaEncoder::Initialize()
 	}
 
 	// Specifying 0 will completely disable frame skipping (therefore encoding as many frames as possible)
+	PRAGMA_DISABLE_DEPRECATION_WARNINGS
 	FParse::Value(FCommandLine::Get(), TEXT("GameplayMediaEncoder.FPS="), VideoConfig.Framerate);
 	if(VideoConfig.Framerate == 0)
+	PRAGMA_ENABLE_DEPRECATION_WARNINGS
 	{
 		// Note : When disabling frame skipping, we lie to the encoder when initializing.
 		// We still specify a framerate, but then feed frames without skipping
+		PRAGMA_DISABLE_DEPRECATION_WARNINGS
 		VideoConfig.Framerate = HardcodedVideoFPS;
+		PRAGMA_ENABLE_DEPRECATION_WARNINGS
 		bDoFrameSkipping = false;
 		UE_LOG(GameplayMediaEncoder, Log, TEXT("Uncapping FPS"));
 	}
 	else
 	{
+		PRAGMA_DISABLE_DEPRECATION_WARNINGS
 		VideoConfig.Framerate = FMath::Clamp(VideoConfig.Framerate, (uint32)MinVideoFPS, (uint32)MaxVideoFPS);
-		bDoFrameSkipping = true;
 		UE_LOG(GameplayMediaEncoder, Log, TEXT("Capping FPS %u"), VideoConfig.Framerate);
+		PRAGMA_ENABLE_DEPRECATION_WARNINGS
+		bDoFrameSkipping = true;
 	}
 
+	PRAGMA_DISABLE_DEPRECATION_WARNINGS
 	VideoConfig.Bitrate = HardcodedVideoBitrate;
 	FParse::Value(FCommandLine::Get(), TEXT("GameplayMediaEncoder.Bitrate="), VideoConfig.Bitrate);
 	VideoConfig.Bitrate = FMath::Clamp(VideoConfig.Bitrate, (uint32)MinVideoBitrate, (uint32)MaxVideoBitrate);
@@ -216,6 +242,7 @@ bool FGameplayMediaEncoder::Initialize()
 	videoInit.MaxBitrate = MaxVideoBitrate;
 	videoInit.TargetBitrate = VideoConfig.Bitrate;
 	videoInit.MaxFramerate = VideoConfig.Framerate;
+	PRAGMA_ENABLE_DEPRECATION_WARNINGS
 
 	if(GDynamicRHI)
 	{
@@ -224,16 +251,25 @@ bool FGameplayMediaEncoder::Initialize()
 #if PLATFORM_DESKTOP && !PLATFORM_APPLE
 		if (RHIType == ERHIInterfaceType::D3D11)
 		{
+			PRAGMA_DISABLE_DEPRECATION_WARNINGS
 			VideoEncoderInput = AVEncoder::FVideoEncoderInput::CreateForD3D11(GDynamicRHI->RHIGetNativeDevice(), true, IsRHIDeviceAMD());
+			PRAGMA_ENABLE_DEPRECATION_WARNINGS
 		}
+		
 		else if (RHIType == ERHIInterfaceType::D3D12)
 		{
+			PRAGMA_DISABLE_DEPRECATION_WARNINGS
 			VideoEncoderInput = AVEncoder::FVideoEncoderInput::CreateForD3D12(GDynamicRHI->RHIGetNativeDevice(), true, IsRHIDeviceNVIDIA());
+			PRAGMA_ENABLE_DEPRECATION_WARNINGS
 		}
+		
 		else if (RHIType == ERHIInterfaceType::Vulkan)
 		{
+			PRAGMA_DISABLE_DEPRECATION_WARNINGS
 			VideoEncoderInput = AVEncoder::FVideoEncoderInput::CreateForVulkan(GDynamicRHI->RHIGetNativeDevice(), true);
+			PRAGMA_ENABLE_DEPRECATION_WARNINGS
 		}
+		
 		else
 #endif
 		{
@@ -242,7 +278,9 @@ bool FGameplayMediaEncoder::Initialize()
 		}
 	}
 
+	PRAGMA_DISABLE_DEPRECATION_WARNINGS
 	const TArray<AVEncoder::FVideoEncoderInfo>& AvailableEncodersInfo = AVEncoder::FVideoEncoderFactory::Get().GetAvailable();
+	PRAGMA_ENABLE_DEPRECATION_WARNINGS
 
 	if (AvailableEncodersInfo.Num() == 0)
 	{
@@ -252,10 +290,12 @@ bool FGameplayMediaEncoder::Initialize()
 
 	for (const auto& EncoderInfo : AvailableEncodersInfo)
 	{
+		PRAGMA_DISABLE_DEPRECATION_WARNINGS
 		if (EncoderInfo.CodecType == AVEncoder::ECodecType::H264)
 		{
 			VideoEncoder = AVEncoder::FVideoEncoderFactory::Get().Create(EncoderInfo.ID, VideoEncoderInput, videoInit);
 		}
+		PRAGMA_ENABLE_DEPRECATION_WARNINGS
 	}
 
 	if (!VideoEncoder)
@@ -264,8 +304,10 @@ bool FGameplayMediaEncoder::Initialize()
 		return false;
 	}
 
+	PRAGMA_DISABLE_DEPRECATION_WARNINGS
 	VideoEncoder->SetOnEncodedPacket([this](uint32 LayerIndex, const TSharedPtr<AVEncoder::FVideoEncoderInputFrame> Frame, const AVEncoder::FCodecPacket& Packet)
 	                                 { OnEncodedVideoFrame(LayerIndex, Frame, Packet); });
+	PRAGMA_ENABLE_DEPRECATION_WARNINGS
 
 	if(!VideoEncoder)
 	{
@@ -311,7 +353,9 @@ bool FGameplayMediaEncoder::Start()
 		AudioDevice->RegisterSubmixBufferListener(AsShared(), AudioDevice->GetMainSubmixObject());
 	}
 
+	PRAGMA_DISABLE_DEPRECATION_WARNINGS
 	FSlateApplication::Get().GetRenderer()->OnBackBufferReadyToPresent().AddRaw(this, &FGameplayMediaEncoder::OnFrameBufferReady);
+	PRAGMA_ENABLE_DEPRECATION_WARNINGS
 
 	return true;
 }
@@ -343,8 +387,11 @@ void FGameplayMediaEncoder::Stop()
 	AudioClock = 0;
 }
 
+PRAGMA_DISABLE_DEPRECATION_WARNINGS
 AVEncoder::FAudioConfig FGameplayMediaEncoder::GetAudioConfig() const
+PRAGMA_ENABLE_DEPRECATION_WARNINGS
 {
+	PRAGMA_DISABLE_DEPRECATION_WARNINGS
 	if(AudioEncoder)
 	{
 		auto codec = AudioEncoder->GetType();
@@ -355,6 +402,7 @@ AVEncoder::FAudioConfig FGameplayMediaEncoder::GetAudioConfig() const
 	{
 		return {};
 	}
+	PRAGMA_ENABLE_DEPRECATION_WARNINGS
 }
 
 void FGameplayMediaEncoder::Shutdown()
@@ -370,7 +418,9 @@ void FGameplayMediaEncoder::Shutdown()
 		if(AudioEncoder)
 		{
 			// AudioEncoder->Reset();
+			PRAGMA_DISABLE_DEPRECATION_WARNINGS
 			AudioEncoder->Shutdown();
+			PRAGMA_ENABLE_DEPRECATION_WARNINGS
 			AudioEncoder.Reset();
 		}
 	}
@@ -378,7 +428,9 @@ void FGameplayMediaEncoder::Shutdown()
 		FScopeLock Lock(&VideoProcessingCS);
 		if(VideoEncoder)
 		{
+			PRAGMA_DISABLE_DEPRECATION_WARNINGS
 			VideoEncoder->Shutdown();
+			PRAGMA_ENABLE_DEPRECATION_WARNINGS
 			VideoEncoder.Reset();
 
 			BackBuffers.Empty();
@@ -388,7 +440,9 @@ void FGameplayMediaEncoder::Shutdown()
 
 FTimespan FGameplayMediaEncoder::GetMediaTimestamp() const { return FTimespan::FromSeconds(FPlatformTime::Seconds()) - StartTime; }
 
+PRAGMA_DISABLE_DEPRECATION_WARNINGS
 const FString& FGameplayMediaEncoder::GetListenerName() const
+PRAGMA_ENABLE_DEPRECATION_WARNINGS
 {
 	static const FString& ListenerName = TEXT("GameplayMediaEncoderListener");
 	return ListenerName;
@@ -531,12 +585,14 @@ void FGameplayMediaEncoder::ProcessAudioFrame(const float* AudioData, int32 NumS
 	// auto encodedInfo = AudioEncoder->Encode(timestamp, audio, &encoded);
 	// OnEncodedAudioFrame(encodedInfo, &encoded);
 
+	PRAGMA_DISABLE_DEPRECATION_WARNINGS
 	AVEncoder::FAudioFrame Frame;
 	Frame.Timestamp = FTimespan::FromSeconds(AudioClock);
 	Frame.Duration = FTimespan::FromSeconds(FloatBuffer.GetSampleDuration());
 	FloatBuffer.Clamp();
 	Frame.Data = FloatBuffer;
 	AudioEncoder->Encode(Frame);
+	PRAGMA_ENABLE_DEPRECATION_WARNINGS
 
 	AudioClock += FloatBuffer.GetSampleDuration();
 }
@@ -555,7 +611,9 @@ void FGameplayMediaEncoder::ProcessVideoFrame(const FTextureRHIRef& FrameBuffer)
 
 	if(bDoFrameSkipping)
 	{
+		PRAGMA_DISABLE_DEPRECATION_WARNINGS
 		uint64 NumExpectedFrames = static_cast<uint64>(Now.GetTotalSeconds() * VideoConfig.Framerate);
+		PRAGMA_ENABLE_DEPRECATION_WARNINGS
 		UE_LOG(GameplayMediaEncoder, VeryVerbose, TEXT("time %.3f: captured %d, expected %d"), Now.GetTotalSeconds(), NumCapturedFrames + 1, NumExpectedFrames);
 		if(NumCapturedFrames + 1 > NumExpectedFrames)
 		{
@@ -566,24 +624,32 @@ void FGameplayMediaEncoder::ProcessVideoFrame(const FTextureRHIRef& FrameBuffer)
 
 	UpdateVideoConfig();
 
+	PRAGMA_DISABLE_DEPRECATION_WARNINGS
 	TSharedPtr<AVEncoder::FVideoEncoderInputFrame> InputFrame = ObtainInputFrame();
 	const int32 FrameId = InputFrame->GetFrameID();
 	InputFrame->SetTimestampUs(Now.GetTicks());
+	PRAGMA_ENABLE_DEPRECATION_WARNINGS
 
 	CopyTexture(FrameBuffer, BackBuffers[InputFrame]);
 
+	PRAGMA_DISABLE_DEPRECATION_WARNINGS
 	AVEncoder::FVideoEncoder::FEncodeOptions EncodeOptions;
 	VideoEncoder->Encode(InputFrame, EncodeOptions);
+	PRAGMA_ENABLE_DEPRECATION_WARNINGS
 
 	LastVideoInputTimestamp = Now;
 	NumCapturedFrames++;
 }
 
+PRAGMA_DISABLE_DEPRECATION_WARNINGS
 TSharedPtr<AVEncoder::FVideoEncoderInputFrame> FGameplayMediaEncoder::ObtainInputFrame()
+PRAGMA_ENABLE_DEPRECATION_WARNINGS
 {
+	PRAGMA_DISABLE_DEPRECATION_WARNINGS
 	TSharedPtr<AVEncoder::FVideoEncoderInputFrame> InputFrame = VideoEncoderInput->ObtainInputFrame();
 	InputFrame->SetWidth(VideoConfig.Width);
 	InputFrame->SetHeight(VideoConfig.Height);
+	PRAGMA_ENABLE_DEPRECATION_WARNINGS
 
 	if(!BackBuffers.Contains(InputFrame))
 	{
@@ -591,7 +657,9 @@ TSharedPtr<AVEncoder::FVideoEncoderInputFrame> FGameplayMediaEncoder::ObtainInpu
 		const ERHIInterfaceType RHIType = RHIGetInterfaceType();
 
 		const FRHITextureCreateDesc Desc =
+			PRAGMA_DISABLE_DEPRECATION_WARNINGS
 			FRHITextureCreateDesc::Create2D(TEXT("VideoCapturerBackBuffer"), VideoConfig.Width, VideoConfig.Height, PF_B8G8R8A8)
+			PRAGMA_ENABLE_DEPRECATION_WARNINGS
 			.SetFlags(ETextureCreateFlags::Shared | ETextureCreateFlags::RenderTargetable | ETextureCreateFlags::UAV)
 			.SetInitialState(ERHIAccess::CopyDest);
 
@@ -600,7 +668,9 @@ TSharedPtr<AVEncoder::FVideoEncoderInputFrame> FGameplayMediaEncoder::ObtainInpu
 			FRHIResourceCreateInfo CreateInfo(TEXT("VideoCapturerBackBuffer"));
 			FTextureRHIRef Texture = RHICreateTexture(Desc);
 
+			PRAGMA_DISABLE_DEPRECATION_WARNINGS
 			InputFrame->SetTexture((ID3D11Texture2D*)Texture->GetNativeResource(), [&, InputFrame](ID3D11Texture2D* NativeTexture) { BackBuffers.Remove(InputFrame); });
+			PRAGMA_ENABLE_DEPRECATION_WARNINGS
 			BackBuffers.Add(InputFrame, Texture);
 		}
 		else if (RHIType == ERHIInterfaceType::D3D12)
@@ -608,7 +678,9 @@ TSharedPtr<AVEncoder::FVideoEncoderInputFrame> FGameplayMediaEncoder::ObtainInpu
 			FRHIResourceCreateInfo CreateInfo(TEXT("VideoCapturerBackBuffer"));
 			FTextureRHIRef Texture = RHICreateTexture(Desc);
 
+			PRAGMA_DISABLE_DEPRECATION_WARNINGS
 			InputFrame->SetTexture((ID3D12Resource*)Texture->GetNativeResource(), [&, InputFrame](ID3D12Resource* NativeTexture) { BackBuffers.Remove(InputFrame); });
+			PRAGMA_ENABLE_DEPRECATION_WARNINGS
 			BackBuffers.Add(InputFrame, Texture);
 		}
 
@@ -637,27 +709,37 @@ void FGameplayMediaEncoder::UpdateVideoConfig()
 {
 	if(bChangeBitrate || bChangeFramerate)
 	{
+		PRAGMA_DISABLE_DEPRECATION_WARNINGS
 		auto config = VideoEncoder->GetLayerConfig(0);
+		PRAGMA_ENABLE_DEPRECATION_WARNINGS
 
 		if(bChangeBitrate)
 		{
+			PRAGMA_DISABLE_DEPRECATION_WARNINGS
 			config.MaxBitrate = MaxVideoBitrate;
 			config.TargetBitrate = NewVideoBitrate;
+			PRAGMA_ENABLE_DEPRECATION_WARNINGS
 		}
 
 		if(bChangeFramerate)
 		{
+			PRAGMA_DISABLE_DEPRECATION_WARNINGS
 			config.MaxFramerate = NewVideoFramerate;
+			PRAGMA_ENABLE_DEPRECATION_WARNINGS
 			NumCapturedFrames = 0;
 		}
 
+		PRAGMA_DISABLE_DEPRECATION_WARNINGS
 		VideoEncoder->UpdateLayerConfig(0, config);
+		PRAGMA_ENABLE_DEPRECATION_WARNINGS
 		bChangeFramerate = false;
 		bChangeBitrate = false;
 	}
 }
 
+PRAGMA_DISABLE_DEPRECATION_WARNINGS
 void FGameplayMediaEncoder::OnEncodedAudioFrame(const AVEncoder::FMediaPacket& Packet)
+PRAGMA_ENABLE_DEPRECATION_WARNINGS
 {
 
 	FScopeLock Lock(&ListenersCS);
@@ -667,8 +749,11 @@ void FGameplayMediaEncoder::OnEncodedAudioFrame(const AVEncoder::FMediaPacket& P
 	}
 }
 
+PRAGMA_DISABLE_DEPRECATION_WARNINGS
 void FGameplayMediaEncoder::OnEncodedVideoFrame(uint32 LayerIndex, const TSharedPtr<AVEncoder::FVideoEncoderInputFrame> InputFrame, const AVEncoder::FCodecPacket& Packet)
+PRAGMA_ENABLE_DEPRECATION_WARNINGS
 {
+	PRAGMA_DISABLE_DEPRECATION_WARNINGS
 	AVEncoder::FMediaPacket packet(AVEncoder::EPacketType::Video);
 
 	packet.Timestamp = InputFrame->GetTimestampUs();
@@ -679,6 +764,7 @@ void FGameplayMediaEncoder::OnEncodedVideoFrame(uint32 LayerIndex, const TShared
 	packet.Video.Height = InputFrame->GetHeight();
 	packet.Video.FrameAvgQP = Packet.VideoQP;
 	packet.Video.Framerate = VideoConfig.Framerate;
+	PRAGMA_ENABLE_DEPRECATION_WARNINGS
 
 	FScopeLock Lock(&ListenersCS);
 	for(auto&& Listener : Listeners)
@@ -686,7 +772,9 @@ void FGameplayMediaEncoder::OnEncodedVideoFrame(uint32 LayerIndex, const TShared
 		Listener->OnMediaSample(packet);
 	}
 
+	PRAGMA_DISABLE_DEPRECATION_WARNINGS
 	InputFrame->Release();
+	PRAGMA_ENABLE_DEPRECATION_WARNINGS
 }
 
 void FGameplayMediaEncoder::CopyTexture(const FTextureRHIRef& SourceTexture, FTextureRHIRef& DestinationTexture) const
diff --git a/Engine/Source/Runtime/GameplayMediaEncoder/Private/GameplayMediaEncoderModule.cpp b/Engine/Source/Runtime/GameplayMediaEncoder/Private/GameplayMediaEncoderModule.cpp
--- a/Engine/Source/Runtime/GameplayMediaEncoder/Private/GameplayMediaEncoderModule.cpp
+++ b/Engine/Source/Runtime/GameplayMediaEncoder/Private/GameplayMediaEncoderModule.cpp
@@ -27,8 +27,11 @@ public:
 		// If the FGameplayMediaEncoder instance was created, then explicitly destroy it here
 		// instead of waiting for the automatic cleanup, since at that point some objects
 		// it depends to for a clean shutdown are not available any longer.
+		PRAGMA_DISABLE_DEPRECATION_WARNINGS
 		FGameplayMediaEncoder::Singleton.Reset();
+		PRAGMA_ENABLE_DEPRECATION_WARNINGS
 	}
+	
 private:
 };
 
diff --git a/Engine/Source/Runtime/GameplayMediaEncoder/Public/GameplayMediaEncoder.h b/Engine/Source/Runtime/GameplayMediaEncoder/Public/GameplayMediaEncoder.h
--- a/Engine/Source/Runtime/GameplayMediaEncoder/Public/GameplayMediaEncoder.h
+++ b/Engine/Source/Runtime/GameplayMediaEncoder/Public/GameplayMediaEncoder.h
@@ -32,10 +32,14 @@ class SWindow;
 class IGameplayMediaEncoderListener
 {
 public:
+	PRAGMA_DISABLE_DEPRECATION_WARNINGS
 	virtual void OnMediaSample(const AVEncoder::FMediaPacket& Sample) = 0;
+	PRAGMA_ENABLE_DEPRECATION_WARNINGS
 };
 
-class FGameplayMediaEncoder final : public ISubmixBufferListener, public AVEncoder::IAudioEncoderListener
+PRAGMA_DISABLE_DEPRECATION_WARNINGS
+class UE_DEPRECATED(5.4, "GameplayMediaEncoder will move to a plugin in a future Unreal Engine release.") FGameplayMediaEncoder final : public ISubmixBufferListener, public AVEncoder::IAudioEncoderListener
+PRAGMA_ENABLE_DEPRECATION_WARNINGS
 {
 public:
 
@@ -87,8 +91,10 @@ public:
 		Get()->Stop();
 	}
 
+	PRAGMA_DISABLE_DEPRECATION_WARNINGS
 	GAMEPLAYMEDIAENCODER_API AVEncoder::FAudioConfig GetAudioConfig() const;
 	AVEncoder::FVideoConfig GetVideoConfig() const { return VideoConfig; }
+	PRAGMA_ENABLE_DEPRECATION_WARNINGS
 
 private:
 	template <typename ObjectType, ESPMode Mode>
@@ -112,10 +118,12 @@ private:
 
 	GAMEPLAYMEDIAENCODER_API void UpdateVideoConfig();
 
+	PRAGMA_DISABLE_DEPRECATION_WARNINGS
 	GAMEPLAYMEDIAENCODER_API void OnEncodedAudioFrame(const AVEncoder::FMediaPacket& Packet) override;
 	GAMEPLAYMEDIAENCODER_API void OnEncodedVideoFrame(uint32 LayerIndex, const TSharedPtr<AVEncoder::FVideoEncoderInputFrame> Frame, const AVEncoder::FCodecPacket& Packet);
 
 	GAMEPLAYMEDIAENCODER_API TSharedPtr<AVEncoder::FVideoEncoderInputFrame> ObtainInputFrame();
+	PRAGMA_ENABLE_DEPRECATION_WARNINGS
 	GAMEPLAYMEDIAENCODER_API void CopyTexture(const FTextureRHIRef& SourceTexture, FTextureRHIRef& DestinationTexture) const;
 
 	GAMEPLAYMEDIAENCODER_API void FloatToPCM16(float const* floatSamples, int32 numSamples, TArray<int16>& out) const;
@@ -126,12 +134,14 @@ private:
 	FCriticalSection AudioProcessingCS;
 	FCriticalSection VideoProcessingCS;
 
+	PRAGMA_DISABLE_DEPRECATION_WARNINGS
 	TUniquePtr<AVEncoder::FAudioEncoder> AudioEncoder;
 
 	AVEncoder::FVideoConfig VideoConfig;
 
 	TUniquePtr<AVEncoder::FVideoEncoder> VideoEncoder;
 	TSharedPtr<AVEncoder::FVideoEncoderInput> VideoEncoderInput;
+	PRAGMA_ENABLE_DEPRECATION_WARNINGS
 
 	uint64 NumCapturedFrames = 0;
 	FTimespan StartTime = 0;
@@ -156,6 +166,9 @@ private:
 	FThreadSafeBool bChangeFramerate = false;
 
 	TArray<int16> PCM16;
+	PRAGMA_DISABLE_DEPRECATION_WARNINGS
 	TMap<TSharedPtr<AVEncoder::FVideoEncoderInputFrame>, FTextureRHIRef> BackBuffers;
+	PRAGMA_ENABLE_DEPRECATION_WARNINGS
 };
 
+
diff --git a/Engine/Source/Runtime/Windows/WindowsPlatformFeatures/Private/HighlightRecorder.cpp b/Engine/Source/Runtime/Windows/WindowsPlatformFeatures/Private/HighlightRecorder.cpp
--- a/Engine/Source/Runtime/Windows/WindowsPlatformFeatures/Private/HighlightRecorder.cpp
+++ b/Engine/Source/Runtime/Windows/WindowsPlatformFeatures/Private/HighlightRecorder.cpp
@@ -66,7 +66,9 @@ bool FHighlightRecorder::Start(double RingBufferDurationSecs)
 	TotalPausedDuration = 0;
 	NumPushedFrames = 0;
 
+	PRAGMA_DISABLE_DEPRECATION_WARNINGS
 	if (!FGameplayMediaEncoder::Get()->RegisterListener(this))
+	PRAGMA_ENABLE_DEPRECATION_WARNINGS
 	{
 		return false;
 	}
@@ -112,7 +114,9 @@ void FHighlightRecorder::Stop()
 {
 	CSV_SCOPED_TIMING_STAT(WindowsVideoRecordingSystem, HighlightRecorder_Stop);
 
+	PRAGMA_DISABLE_DEPRECATION_WARNINGS
 	FGameplayMediaEncoder::Get()->UnregisterListener(this);
+	PRAGMA_ENABLE_DEPRECATION_WARNINGS
 	State = EState::Stopped;
 
 	if (BackgroundSaving)
@@ -129,6 +133,7 @@ FTimespan FHighlightRecorder::GetRecordingTime() const
 	return FTimespan::FromSeconds(FPlatformTime::Seconds()) - RecordingStartTime - TotalPausedDuration;
 }
 
+PRAGMA_DISABLE_DEPRECATION_WARNINGS
 void FHighlightRecorder::OnMediaSample(const AVEncoder::FMediaPacket& InSample)
 {
 	// We might be paused, so don't do anything
@@ -165,6 +170,7 @@ void FHighlightRecorder::OnMediaSample(const AVEncoder::FMediaPacket& InSample)
 
 	RingBuffer.Push(MoveTemp(SampleCopy));
 }
+PRAGMA_ENABLE_DEPRECATION_WARNINGS
 
 bool FHighlightRecorder::SaveHighlight(const TCHAR* Filename, FDoneCallback InDoneCallback, double MaxDurationSecs)
 {
@@ -226,9 +232,13 @@ bool FHighlightRecorder::SaveHighlightInBackground(const FString& Filename, doub
 	return bRes;
 }
 
+
 bool FHighlightRecorder::SaveHighlightInBackgroundImpl(const FString& Filename, double MaxDurationSecs)
 {
+	PRAGMA_DISABLE_DEPRECATION_WARNINGS
 	TArray<AVEncoder::FMediaPacket> Samples = RingBuffer.GetCopy();
+	PRAGMA_ENABLE_DEPRECATION_WARNINGS
+	
 
 	if (Samples.Num()==0)
 	{
@@ -247,7 +257,9 @@ bool FHighlightRecorder::SaveHighlightInBackgroundImpl(const FString& Filename,
 	bool bHasAudio = false;
 	for (int Idx = FirstSampleIndex; Idx != Samples.Num(); ++Idx)
 	{
+		PRAGMA_DISABLE_DEPRECATION_WARNINGS
 		if (Samples[Idx].Type == AVEncoder::EPacketType::Audio)
+		PRAGMA_ENABLE_DEPRECATION_WARNINGS
 		{
 			bHasAudio = true;
 			break;
@@ -259,7 +271,9 @@ bool FHighlightRecorder::SaveHighlightInBackgroundImpl(const FString& Filename,
 		return false;
 	}
 
+	PRAGMA_DISABLE_DEPRECATION_WARNINGS
 	checkf(Samples[FirstSampleIndex].IsVideoKeyFrame(), TEXT("t %.3f d %.3f"), Samples[FirstSampleIndex].Timestamp.GetTotalSeconds(), Samples[FirstSampleIndex].Duration.GetTotalSeconds());
+	PRAGMA_ENABLE_DEPRECATION_WARNINGS
 
 	if (FirstSampleIndex == Samples.Num())
 	{
@@ -267,18 +281,22 @@ bool FHighlightRecorder::SaveHighlightInBackgroundImpl(const FString& Filename,
 		return false;
 	}
 
+	PRAGMA_DISABLE_DEPRECATION_WARNINGS
 	UE_LOG(HighlightRecorder, Verbose, TEXT("writting %d samples to .mp4, %.3f s, starting from %.3f s, index %d"),
 		Samples.Num() - FirstSampleIndex,
 		(Samples.Last().Timestamp - StartTime + Samples.Last().Duration).GetTotalSeconds(),
 		StartTime.GetTotalSeconds(),
 		FirstSampleIndex);
+	PRAGMA_ENABLE_DEPRECATION_WARNINGS
 
 	// get samples starting from `StartTime` and push them into Mp4Writer
 	for (int Idx = FirstSampleIndex; Idx != Samples.Num(); ++Idx)
 	{
+		PRAGMA_DISABLE_DEPRECATION_WARNINGS
 		AVEncoder::FMediaPacket& Sample = Samples[Idx];
 		Sample.Timestamp = Sample.Timestamp - StartTime;
 		if (!Mp4Writer->Write(Sample, (Sample.Type==AVEncoder::EPacketType::Audio) ? AudioStreamIndex : VideoStreamIndex))
+		PRAGMA_ENABLE_DEPRECATION_WARNINGS
 		{
 			return false;
 		}
@@ -292,6 +310,7 @@ bool FHighlightRecorder::SaveHighlightInBackgroundImpl(const FString& Filename,
 	return true;
 }
 
+
 bool FHighlightRecorder::InitialiseMp4Writer(const FString& FullFilename, bool bHasAudio)
 {
 	// create target directory if it does not exist
@@ -316,7 +335,9 @@ bool FHighlightRecorder::InitialiseMp4Writer(const FString& FullFilename, bool b
 
 	if (bHasAudio)
 	{
+		PRAGMA_DISABLE_DEPRECATION_WARNINGS
 		TPair<FString, AVEncoder::FAudioConfig> AudioConfig = TPair<FString, AVEncoder::FAudioConfig>( FGameplayMediaEncoder::Get()->GetAudioConfig().Codec, FGameplayMediaEncoder::Get()->GetAudioConfig());		
+		PRAGMA_ENABLE_DEPRECATION_WARNINGS
 		if (AudioConfig.Key == "")
 		{
 			UE_LOG(HighlightRecorder, Error, TEXT("Could not get audio config"));
@@ -334,7 +355,9 @@ bool FHighlightRecorder::InitialiseMp4Writer(const FString& FullFilename, bool b
 		}
 	}
 
+	PRAGMA_DISABLE_DEPRECATION_WARNINGS
 	TPair<FString, AVEncoder::FVideoConfig> VideoConfig = TPair<FString, AVEncoder::FVideoConfig>( FGameplayMediaEncoder::Get()->GetVideoConfig().Codec, FGameplayMediaEncoder::Get()->GetVideoConfig());
+	PRAGMA_ENABLE_DEPRECATION_WARNINGS
 	if (VideoConfig.Key == "")
 	{
 		UE_LOG(HighlightRecorder, Error, TEXT("Could not get video config"));
@@ -360,6 +383,7 @@ bool FHighlightRecorder::InitialiseMp4Writer(const FString& FullFilename, bool b
 }
 
 // finds index and timestamp of the first sample that should be written to .mp4
+PRAGMA_DISABLE_DEPRECATION_WARNINGS
 bool FHighlightRecorder::GetSavingStart(const TArray<AVEncoder::FMediaPacket>& Samples, FTimespan MaxDuration, int& StartIndex, FTimespan& StartTime) const
 // the first sample in .mp4 file should have timestamp 0 and all other timestamps should be relative to the
 // first one
@@ -411,3 +435,4 @@ bool FHighlightRecorder::GetSavingStart(const TArray<AVEncoder::FMediaPacket>& S
 
 	return true;
 }
+PRAGMA_ENABLE_DEPRECATION_WARNINGS
\ No newline at end of file
diff --git a/Engine/Source/Runtime/Windows/WindowsPlatformFeatures/Private/HighlightRecorder.h b/Engine/Source/Runtime/Windows/WindowsPlatformFeatures/Private/HighlightRecorder.h
--- a/Engine/Source/Runtime/Windows/WindowsPlatformFeatures/Private/HighlightRecorder.h
+++ b/Engine/Source/Runtime/Windows/WindowsPlatformFeatures/Private/HighlightRecorder.h
@@ -45,7 +45,9 @@ private:
 	bool SaveHighlightInBackground(const FString& Filename, double MaxDurationSecs);
 	bool SaveHighlightInBackgroundImpl(const FString& Filename, double MaxDurationSecs);
 	bool InitialiseMp4Writer(const FString& Filename, bool bHasAudio);
+	PRAGMA_DISABLE_DEPRECATION_WARNINGS
 	bool GetSavingStart(const TArray<AVEncoder::FMediaPacket>& Samples, FTimespan MaxDuration, int& OutStartIndex, FTimespan& OutStartTime) const;
+	PRAGMA_ENABLE_DEPRECATION_WARNINGS
 
 	// takes into account if we've been paused and shifts current time back to compensate paused state
 	// so all timestamps are continuous even over paused pieces
@@ -54,7 +56,9 @@ private:
 	//
 	// IGameplayMediaEncoderListener implementation
 	//
+	PRAGMA_DISABLE_DEPRECATION_WARNINGS
 	void OnMediaSample(const AVEncoder::FMediaPacket& Sample) override;
+	PRAGMA_ENABLE_DEPRECATION_WARNINGS
 
 private:
 	TAtomic<EState> State{ EState::Stopped };
diff --git a/Engine/Source/Runtime/Windows/WindowsPlatformFeatures/Private/WindowsPlatformFeatures.cpp b/Engine/Source/Runtime/Windows/WindowsPlatformFeatures/Private/WindowsPlatformFeatures.cpp
--- a/Engine/Source/Runtime/Windows/WindowsPlatformFeatures/Private/WindowsPlatformFeatures.cpp
+++ b/Engine/Source/Runtime/Windows/WindowsPlatformFeatures/Private/WindowsPlatformFeatures.cpp
@@ -11,12 +11,17 @@ IMPLEMENT_MODULE(FWindowsPlatformFeaturesModule, WindowsPlatformFeatures);
 
 FWindowsPlatformFeaturesModule::FWindowsPlatformFeaturesModule()
 {
+	VideoRecordingSystem = MakeShared<FWindowsVideoRecordingSystem>();
 }
 
 IVideoRecordingSystem* FWindowsPlatformFeaturesModule::GetVideoRecordingSystem()
 {
-	static FWindowsVideoRecordingSystem VideoRecordingSystem;
-	return &VideoRecordingSystem;
+	return VideoRecordingSystem.Get();
+}
+
+void FWindowsPlatformFeaturesModule::RegisterVideoRecordingSystem(TSharedPtr<IVideoRecordingSystem> InVideoRecordingSystem)
+{
+	VideoRecordingSystem = InVideoRecordingSystem;
 }
 
 ISaveGameSystem* FWindowsPlatformFeaturesModule::GetSaveGameSystem()
diff --git a/Engine/Source/Runtime/Windows/WindowsPlatformFeatures/Private/WindowsVideoRecordingSystem.cpp b/Engine/Source/Runtime/Windows/WindowsPlatformFeatures/Private/WindowsVideoRecordingSystem.cpp
--- a/Engine/Source/Runtime/Windows/WindowsPlatformFeatures/Private/WindowsVideoRecordingSystem.cpp
+++ b/Engine/Source/Runtime/Windows/WindowsPlatformFeatures/Private/WindowsVideoRecordingSystem.cpp
@@ -155,19 +155,23 @@ private:
 
 
 
+PRAGMA_DISABLE_DEPRECATION_WARNINGS
 FWindowsVideoRecordingSystem::FWindowsVideoRecordingSystem()
 {
 	UE_LOG(WindowsVideoRecordingSystem, Verbose, TEXT("%s"), __FUNCTIONW__);
 
 	EnableRecording(true);
 }
+PRAGMA_ENABLE_DEPRECATION_WARNINGS
 
+PRAGMA_DISABLE_DEPRECATION_WARNINGS
 FWindowsVideoRecordingSystem::~FWindowsVideoRecordingSystem()
 {
 	UE_LOG(WindowsVideoRecordingSystem, Verbose, TEXT("%s"), __FUNCTIONW__);
 
 	EnableRecording(false);
 }
+PRAGMA_ENABLE_DEPRECATION_WARNINGS
 
 void FWindowsVideoRecordingSystem::EnableRecording(bool bEnableRecording)
 {
@@ -405,8 +409,11 @@ void FWindowsVideoRecordingSystem::FinalizeCallbackOnGameThread(bool bSaved, boo
 
 	if (bBroadcast)
 	{
+		PRAGMA_DISABLE_DEPRECATION_WARNINGS
 		OnVideoRecordingFinalized.Broadcast(bSaved, Path);
+		PRAGMA_ENABLE_DEPRECATION_WARNINGS
 	}
+	
 }
 
 EVideoRecordingState FWindowsVideoRecordingSystem::GetRecordingState() const
diff --git a/Engine/Source/Runtime/Windows/WindowsPlatformFeatures/Private/WindowsVideoRecordingSystem.h b/Engine/Source/Runtime/Windows/WindowsPlatformFeatures/Private/WindowsVideoRecordingSystem.h
--- a/Engine/Source/Runtime/Windows/WindowsPlatformFeatures/Private/WindowsVideoRecordingSystem.h
+++ b/Engine/Source/Runtime/Windows/WindowsPlatformFeatures/Private/WindowsVideoRecordingSystem.h
@@ -18,7 +18,9 @@ DECLARE_LOG_CATEGORY_EXTERN(WindowsVideoRecordingSystem, VeryVerbose, VeryVerbos
 
 class FHighlightRecorder;
 
+PRAGMA_DISABLE_DEPRECATION_WARNINGS
 class FWindowsVideoRecordingSystem : public IVideoRecordingSystem
+PRAGMA_ENABLE_DEPRECATION_WARNINGS
 {
 public:
 	FWindowsVideoRecordingSystem();
diff --git a/Engine/Source/Runtime/Windows/WindowsPlatformFeatures/Private/WmfMp4Writer.cpp b/Engine/Source/Runtime/Windows/WindowsPlatformFeatures/Private/WmfMp4Writer.cpp
--- a/Engine/Source/Runtime/Windows/WindowsPlatformFeatures/Private/WmfMp4Writer.cpp
+++ b/Engine/Source/Runtime/Windows/WindowsPlatformFeatures/Private/WmfMp4Writer.cpp
@@ -21,6 +21,7 @@ bool FWmfMp4Writer::Initialize(const TCHAR* Filename)
 	return true;
 }
 
+PRAGMA_DISABLE_DEPRECATION_WARNINGS
 TOptional<DWORD> FWmfMp4Writer::CreateAudioStream(const FString& Codec, const AVEncoder::FAudioConfig& Config)
 {
 	GUID Format;
@@ -58,7 +59,9 @@ TOptional<DWORD> FWmfMp4Writer::CreateAudioStream(const FString& Codec, const AV
 	CHECK_HR_DEFAULT(Writer->SetInputMediaType(StreamIndex, MediaType, nullptr));
 	return TOptional<DWORD>(StreamIndex);
 }
+PRAGMA_ENABLE_DEPRECATION_WARNINGS
 
+PRAGMA_DISABLE_DEPRECATION_WARNINGS
 TOptional<DWORD> FWmfMp4Writer::CreateVideoStream(const FString& Codec, const AVEncoder::FVideoConfig& Config)
 {
 	GUID Format;
@@ -94,6 +97,7 @@ TOptional<DWORD> FWmfMp4Writer::CreateVideoStream(const FString& Codec, const AV
 	CHECK_HR_DEFAULT(Writer->SetInputMediaType(StreamIndex, MediaType, nullptr));
 	return TOptional<DWORD>(StreamIndex);
 }
+PRAGMA_ENABLE_DEPRECATION_WARNINGS
 
 bool FWmfMp4Writer::Start()
 {
@@ -101,6 +105,7 @@ bool FWmfMp4Writer::Start()
 	return true;
 }
 
+PRAGMA_DISABLE_DEPRECATION_WARNINGS
 bool FWmfMp4Writer::Write(const AVEncoder::FMediaPacket& InSample, DWORD StreamIndex)
 {
 	AVEncoder::FIMFSampleWrapper Sample { InSample.Type };
@@ -128,6 +133,7 @@ bool FWmfMp4Writer::Write(const AVEncoder::FMediaPacket& InSample, DWORD StreamI
 
 	return true;
 }
+PRAGMA_ENABLE_DEPRECATION_WARNINGS
 
 bool FWmfMp4Writer::Finalize()
 {
diff --git a/Engine/Source/Runtime/Windows/WindowsPlatformFeatures/Private/WmfMp4Writer.h b/Engine/Source/Runtime/Windows/WindowsPlatformFeatures/Private/WmfMp4Writer.h
--- a/Engine/Source/Runtime/Windows/WindowsPlatformFeatures/Private/WmfMp4Writer.h
+++ b/Engine/Source/Runtime/Windows/WindowsPlatformFeatures/Private/WmfMp4Writer.h
@@ -17,6 +17,7 @@ public:
 	/**
 	 * Create an audio stream and return the its index on success
 	 */
+	PRAGMA_DISABLE_DEPRECATION_WARNINGS
 	TOptional<DWORD> CreateAudioStream(const FString& Codec, const AVEncoder::FAudioConfig& Config);
 
 	/**
@@ -24,8 +25,10 @@ public:
 	 */
 	TOptional<DWORD> CreateVideoStream(const FString& Codec, const AVEncoder::FVideoConfig& Config);
 
-	bool Start();
 	bool Write(const AVEncoder::FMediaPacket& InSample, DWORD StreamIndex);
+	PRAGMA_ENABLE_DEPRECATION_WARNINGS
+
+	bool Start();
 	bool Finalize();
 
 private:
diff --git a/Engine/Source/Runtime/Windows/WindowsPlatformFeatures/Private/WmfRingBuffer.cpp b/Engine/Source/Runtime/Windows/WindowsPlatformFeatures/Private/WmfRingBuffer.cpp
--- a/Engine/Source/Runtime/Windows/WindowsPlatformFeatures/Private/WmfRingBuffer.cpp
+++ b/Engine/Source/Runtime/Windows/WindowsPlatformFeatures/Private/WmfRingBuffer.cpp
@@ -6,6 +6,8 @@
 DECLARE_LOG_CATEGORY_EXTERN(WmfRingBuffer, Log, VeryVerbose);
 DEFINE_LOG_CATEGORY(WmfRingBuffer);
 
+
+PRAGMA_DISABLE_DEPRECATION_WARNINGS
 void FWmfRingBuffer::Push(AVEncoder::FMediaPacket&& Sample)
 {
 	check(MaxDuration != 0);
@@ -64,6 +66,7 @@ void FWmfRingBuffer::Push(AVEncoder::FMediaPacket&& Sample)
 		}
 	}
 }
+PRAGMA_ENABLE_DEPRECATION_WARNINGS
 
 void FWmfRingBuffer::Cleanup()
 // removes a/v samples grouped by the oldest key-frame period
@@ -72,14 +75,18 @@ void FWmfRingBuffer::Cleanup()
 
 	check(Samples.Num() != 0);
 
+	PRAGMA_DISABLE_DEPRECATION_WARNINGS
 	check(Samples[0].IsVideoKeyFrame());
+	PRAGMA_ENABLE_DEPRECATION_WARNINGS
 
 	// find the second key-frame
 	int i = 1; // skip the first sample which is a keyframe
 	bool bFound = false;
 	for (; i != Samples.Num(); ++i)
 	{
+		PRAGMA_DISABLE_DEPRECATION_WARNINGS
 		if (Samples[i].IsVideoKeyFrame())
+		PRAGMA_ENABLE_DEPRECATION_WARNINGS
 		{
 			bFound = true;
 			break;
@@ -89,14 +96,18 @@ void FWmfRingBuffer::Cleanup()
 
 	// remove key-frame period
 	Samples.RemoveAt(0, i, EAllowShrinking::No);
+	PRAGMA_DISABLE_DEPRECATION_WARNINGS
 	check(Samples[0].IsVideoKeyFrame());
+	PRAGMA_ENABLE_DEPRECATION_WARNINGS
 
+	PRAGMA_DISABLE_DEPRECATION_WARNINGS
 	UE_LOG(WmfRingBuffer, VeryVerbose, TEXT("%d samples, %.3f s, %.3f - %.3f:%.3f"),
 		Samples.Num(),
 		GetDuration().GetTotalSeconds(),
 		Samples[0].Timestamp.GetTotalSeconds(),
 		Samples.Last().Timestamp.GetTotalSeconds(),
 		Samples.Last().Duration.GetTotalSeconds());
+	PRAGMA_ENABLE_DEPRECATION_WARNINGS
 }
 
 void FWmfRingBuffer::PauseCleanup(bool bPause)
@@ -113,9 +124,12 @@ FTimespan FWmfRingBuffer::GetDuration() const
 
 	// A duration of a video is not "LastFrameCaptureTime - FirstFrameCaptureTime" , but 
 	// (LastFrameCaptureTime - FirstFrameCaptureTime) + LastFrameDuration.
+	PRAGMA_DISABLE_DEPRECATION_WARNINGS
 	return Samples.Last().Timestamp + Samples.Last().Duration - Samples[0].Timestamp;
+	PRAGMA_ENABLE_DEPRECATION_WARNINGS
 }
 
+PRAGMA_DISABLE_DEPRECATION_WARNINGS
 TArray<AVEncoder::FMediaPacket> FWmfRingBuffer::GetCopy()
 {
 	FScopeLock Lock(&Mutex);
@@ -128,6 +142,7 @@ TArray<AVEncoder::FMediaPacket> FWmfRingBuffer::GetCopy()
 
 	return Copy;
 }
+PRAGMA_ENABLE_DEPRECATION_WARNINGS
 
 void FWmfRingBuffer::Reset()
 {
diff --git a/Engine/Source/Runtime/Windows/WindowsPlatformFeatures/Private/WmfRingBuffer.h b/Engine/Source/Runtime/Windows/WindowsPlatformFeatures/Private/WmfRingBuffer.h
--- a/Engine/Source/Runtime/Windows/WindowsPlatformFeatures/Private/WmfRingBuffer.h
+++ b/Engine/Source/Runtime/Windows/WindowsPlatformFeatures/Private/WmfRingBuffer.h
@@ -22,11 +22,13 @@ public:
 
 	FTimespan GetDuration() const;
 
-	void Push(AVEncoder::FMediaPacket&& Sample);
-
 	void PauseCleanup(bool bPause);
 
+	PRAGMA_DISABLE_DEPRECATION_WARNINGS
+	void Push(AVEncoder::FMediaPacket&& Sample);
+
 	TArray<AVEncoder::FMediaPacket> GetCopy();
+	PRAGMA_ENABLE_DEPRECATION_WARNINGS
 
 	void Reset();
 
@@ -35,8 +37,11 @@ private:
 
 private:
 	FTimespan MaxDuration = 0;
-	TArray<AVEncoder::FMediaPacket> Samples;
 	FCriticalSection Mutex;
 	FThreadSafeBool bCleanupPaused = false;
+
+	PRAGMA_DISABLE_DEPRECATION_WARNINGS
+	TArray<AVEncoder::FMediaPacket> Samples;
+	PRAGMA_ENABLE_DEPRECATION_WARNINGS
 };
 
diff --git a/Engine/Source/Runtime/Windows/WindowsPlatformFeatures/Public/WindowsPlatformFeatures.h b/Engine/Source/Runtime/Windows/WindowsPlatformFeatures/Public/WindowsPlatformFeatures.h
--- a/Engine/Source/Runtime/Windows/WindowsPlatformFeatures/Public/WindowsPlatformFeatures.h
+++ b/Engine/Source/Runtime/Windows/WindowsPlatformFeatures/Public/WindowsPlatformFeatures.h
@@ -15,9 +15,13 @@ public:
 	WINDOWSPLATFORMFEATURES_API virtual IVideoRecordingSystem* GetVideoRecordingSystem() override;
 	WINDOWSPLATFORMFEATURES_API virtual class ISaveGameSystem* GetSaveGameSystem() override;
 
+	WINDOWSPLATFORMFEATURES_API void RegisterVideoRecordingSystem(TSharedPtr<IVideoRecordingSystem> VideoRecordingSystem);
+
 private:
 	WINDOWSPLATFORMFEATURES_API void StartupModule() override;
 
+private:
+	TSharedPtr<IVideoRecordingSystem> VideoRecordingSystem;
 };
 
 
