[Backout] - CL32235200 - CIS Compile Error
Original CL Desc
-----------------------------------------------------------------
QOL: Deprecate AVEncoder (for removal) and its dependencies (to be moved to plugins)

diff --git a/Engine/Plugins/Media/HardwareEncoders/Source/EncoderAMF/Private/Amf_EncoderH264.cpp b/Engine/Plugins/Media/HardwareEncoders/Source/EncoderAMF/Private/Amf_EncoderH264.cpp
--- a/Engine/Plugins/Media/HardwareEncoders/Source/EncoderAMF/Private/Amf_EncoderH264.cpp
+++ b/Engine/Plugins/Media/HardwareEncoders/Source/EncoderAMF/Private/Amf_EncoderH264.cpp
@@ -24,7 +24,6 @@
 
 namespace
 {
-	PRAGMA_DISABLE_DEPRECATION_WARNINGS
 	AMF_VIDEO_ENCODER_RATE_CONTROL_METHOD_ENUM ConvertRateControlModeAMF(AVEncoder::FVideoEncoder::RateControlMode mode)
 	{
 		switch (mode)
@@ -49,7 +48,6 @@ namespace
 			case AVEncoder::FVideoEncoder::H264Profile::AUTO: return AMF_VIDEO_ENCODER_PROFILE_BASELINE;
 		}
 	}
-	PRAGMA_ENABLE_DEPRECATION_WARNINGS
 }
 
 namespace AVEncoder
@@ -98,62 +96,46 @@ namespace AVEncoder
 		});
 	}
 
-	PRAGMA_DISABLE_DEPRECATION_WARNINGS
 	static bool GetEncoderInfo(FAmfCommon& Amf, FVideoEncoderInfo& EncoderInfo);
-	PRAGMA_ENABLE_DEPRECATION_WARNINGS
 
-	PRAGMA_DISABLE_DEPRECATION_WARNINGS
 	bool FVideoEncoderAmf_H264::GetIsAvailable(FVideoEncoderInputImpl& InInput, FVideoEncoderInfo& OutEncoderInfo)
-	PRAGMA_ENABLE_DEPRECATION_WARNINGS
 	{
 		FAmfCommon& Amf = FAmfCommon::Setup();
 		bool bIsAvailable = Amf.GetIsAvailable();
 		if (bIsAvailable)
 		{
-			PRAGMA_DISABLE_DEPRECATION_WARNINGS
 			OutEncoderInfo.CodecType = ECodecType::H264;
-			PRAGMA_ENABLE_DEPRECATION_WARNINGS
 		}
-		
 		return bIsAvailable;
 	}
 
-	PRAGMA_DISABLE_DEPRECATION_WARNINGS
 	void FVideoEncoderAmf_H264::Register(FVideoEncoderFactory& InFactory)
-	PRAGMA_ENABLE_DEPRECATION_WARNINGS
 	{
 		FAmfCommon& Amf = FAmfCommon::Setup();
 		if (Amf.GetIsAvailable() && IsRHIDeviceAMD())
 		{
-			PRAGMA_DISABLE_DEPRECATION_WARNINGS
 			FVideoEncoderInfo	EncoderInfo;
-			PRAGMA_ENABLE_DEPRECATION_WARNINGS
 			if (GetEncoderInfo(Amf, EncoderInfo))
 			{
-				PRAGMA_DISABLE_DEPRECATION_WARNINGS
 				InFactory.Register(EncoderInfo, []() {
 					return TUniquePtr<FVideoEncoder>(new FVideoEncoderAmf_H264());
 				});
-				PRAGMA_ENABLE_DEPRECATION_WARNINGS
 			}
 		}
 	}
 
-	PRAGMA_DISABLE_DEPRECATION_WARNINGS
 	FVideoEncoderAmf_H264::FVideoEncoderAmf_H264()
 		: Amf(FAmfCommon::Setup())
 	{
 		AMFParseCommandLineFlags();
 	}
-	PRAGMA_ENABLE_DEPRECATION_WARNINGS
 
-	PRAGMA_DISABLE_DEPRECATION_WARNINGS
-	FVideoEncoderAmf_H264::~FVideoEncoderAmf_H264() { Shutdown(); }
-	PRAGMA_ENABLE_DEPRECATION_WARNINGS
+	FVideoEncoderAmf_H264::~FVideoEncoderAmf_H264()
+	{
+		Shutdown();
+	}
 
-	PRAGMA_DISABLE_DEPRECATION_WARNINGS
 	bool FVideoEncoderAmf_H264::Setup(TSharedRef<FVideoEncoderInput> input, FLayerConfig const& config)
-	PRAGMA_ENABLE_DEPRECATION_WARNINGS
 	{
 		if (!Amf.GetIsAvailable())
 		{
@@ -166,7 +148,6 @@ namespace AVEncoder
 		ERHIInterfaceType RHIType = ERHIInterfaceType::Hidden;
 
 		// TODO fix initializing contexts
-		PRAGMA_DISABLE_DEPRECATION_WARNINGS
 		FrameFormat = input->GetFrameFormat();
 		switch (FrameFormat)
 		{
@@ -189,7 +170,6 @@ namespace AVEncoder
 			UE_LOG(LogEncoderAMF, Error, TEXT("Frame format %s is not currently supported by Amf Encoder on this platform."), *ToString(FrameFormat));
 			return false;
 		}
-		PRAGMA_ENABLE_DEPRECATION_WARNINGS
 
 		if (!EncoderDevice)
 		{
@@ -208,17 +188,14 @@ namespace AVEncoder
 			}
 			else
 			{
-				PRAGMA_DISABLE_DEPRECATION_WARNINGS
 				FVulkanDataStruct* VulkanData = static_cast<FVulkanDataStruct*>(EncoderDevice);
 				if (!Amf.InitializeContext(RHIGetInterfaceType(), GDynamicRHI->GetName(), VulkanData->VulkanDevice, VulkanData->VulkanInstance, VulkanData->VulkanPhysicalDevice))
 				{
 					UE_LOG(LogEncoderAMF, Error, TEXT("Amf component not initialised"));
 				}
-				PRAGMA_ENABLE_DEPRECATION_WARNINGS
 			}
 		}
 
-		PRAGMA_DISABLE_DEPRECATION_WARNINGS
 		FLayerConfig MutableConfig = config;
 		if (MutableConfig.MaxFramerate == 0)
 		{
@@ -226,12 +203,9 @@ namespace AVEncoder
 		}
 
 		return AddLayer(config);
-		PRAGMA_ENABLE_DEPRECATION_WARNINGS
 	}
 
-	PRAGMA_DISABLE_DEPRECATION_WARNINGS
 	FVideoEncoder::FLayer* FVideoEncoderAmf_H264::CreateLayer(uint32 layerIdx, FLayerConfig const& config)
-	PRAGMA_ENABLE_DEPRECATION_WARNINGS
 	{
 		auto const layer = new FAMFLayer(layerIdx, config, *this);
 		if (!layer->Setup())
@@ -243,18 +217,15 @@ namespace AVEncoder
 		return layer;
 	}
 
-	PRAGMA_DISABLE_DEPRECATION_WARNINGS
-	void FVideoEncoderAmf_H264::DestroyLayer(FLayer* layer) { delete layer; }
-	PRAGMA_ENABLE_DEPRECATION_WARNINGS
+	void FVideoEncoderAmf_H264::DestroyLayer(FLayer* layer)
+	{
+		delete layer;
+	}
 
-	PRAGMA_DISABLE_DEPRECATION_WARNINGS
 	void FVideoEncoderAmf_H264::Encode(const TSharedPtr<FVideoEncoderInputFrame> frame, FEncodeOptions const& options)
-	PRAGMA_ENABLE_DEPRECATION_WARNINGS
 	{
-		PRAGMA_DISABLE_DEPRECATION_WARNINGS
 		const TSharedPtr<FVideoEncoderInputFrameImpl> amfFrame = StaticCastSharedPtr<FVideoEncoderInputFrameImpl>(frame);
 		for (auto& layer : Layers)
-		PRAGMA_ENABLE_DEPRECATION_WARNINGS
 		{
 			FAMFLayer* amfLayer = static_cast<FAMFLayer*>(layer);
 			AMF_RESULT Res = amfLayer->Encode(amfFrame, options);
@@ -267,9 +238,7 @@ namespace AVEncoder
 
 	void FVideoEncoderAmf_H264::Flush()
 	{
-		PRAGMA_DISABLE_DEPRECATION_WARNINGS
 		for (auto&& layer : Layers)
-		PRAGMA_ENABLE_DEPRECATION_WARNINGS
 		{
 			auto const amfLayer = static_cast<FAMFLayer*>(layer);
 			amfLayer->Flush();
@@ -278,23 +247,17 @@ namespace AVEncoder
 
 	void FVideoEncoderAmf_H264::Shutdown()
 	{
-		PRAGMA_DISABLE_DEPRECATION_WARNINGS
 		for (auto&& layer : Layers)
-		PRAGMA_ENABLE_DEPRECATION_WARNINGS
 		{
 			auto const amfLayer = static_cast<FAMFLayer*>(layer);
 			amfLayer->Shutdown();
 			DestroyLayer(amfLayer);
 		}
 
-		PRAGMA_DISABLE_DEPRECATION_WARNINGS
 		Layers.Reset();
-		PRAGMA_ENABLE_DEPRECATION_WARNINGS
 	}
-	
 
 	// --- Amf_EncoderH264::FLayer ------------------------------------------------------------
-	PRAGMA_DISABLE_DEPRECATION_WARNINGS
 	FVideoEncoderAmf_H264::FAMFLayer::FAMFLayer(uint32 layerIdx, FLayerConfig const& config, FVideoEncoderAmf_H264& encoder)
 		: FLayer(config)
 		, Encoder(encoder)
@@ -302,11 +265,10 @@ namespace AVEncoder
 		, LayerIndex(layerIdx)
 	{
 	}
-	PRAGMA_ENABLE_DEPRECATION_WARNINGS
 
-	PRAGMA_DISABLE_DEPRECATION_WARNINGS
-	FVideoEncoderAmf_H264::FAMFLayer::~FAMFLayer() {}
-	PRAGMA_ENABLE_DEPRECATION_WARNINGS
+	FVideoEncoderAmf_H264::FAMFLayer::~FAMFLayer()
+	{
+	}
 
 	bool FVideoEncoderAmf_H264::FAMFLayer::Setup()
 	{
@@ -329,13 +291,10 @@ namespace AVEncoder
 
 		Result = AmfEncoder->SetProperty(AMF_VIDEO_ENCODER_USAGE, AMF_VIDEO_ENCODER_USAGE_LOW_LATENCY);
 
-		PRAGMA_DISABLE_DEPRECATION_WARNINGS
 		AMF_VIDEO_ENCODER_PROFILE_ENUM H264Profile = ConvertH264Profile(CurrentConfig.H264Profile);
-		PRAGMA_ENABLE_DEPRECATION_WARNINGS
 		Result = AmfEncoder->SetProperty(AMF_VIDEO_ENCODER_PROFILE, H264Profile);
 		Result = AmfEncoder->SetProperty(AMF_VIDEO_ENCODER_PROFILE_LEVEL, 51);
 
-		PRAGMA_DISABLE_DEPRECATION_WARNINGS
 		AMFRate frameRate = { CurrentConfig.MaxFramerate, 1 };
 		Result = AmfEncoder->SetProperty(AMF_VIDEO_ENCODER_FRAMERATE, frameRate);
 		CurrentFrameRate = CurrentConfig.MaxFramerate;
@@ -368,20 +327,16 @@ namespace AVEncoder
 		Result = AmfEncoder->Init(AMF_SURFACE_BGRA, CurrentConfig.Width, CurrentConfig.Height);
 		CurrentWidth = CurrentConfig.Width;
 		CurrentHeight = CurrentConfig.Height;
-		PRAGMA_ENABLE_DEPRECATION_WARNINGS
 
 		return Result == AMF_OK;
 	}
 
 	void FVideoEncoderAmf_H264::FAMFLayer::MaybeReconfigure()
 	{
-		PRAGMA_DISABLE_DEPRECATION_WARNINGS
 		FScopeLock lock(&ConfigMutex);
 		if (NeedsReconfigure)
-		PRAGMA_ENABLE_DEPRECATION_WARNINGS
 		{	
 			// Static properties - need ReInit
-			PRAGMA_DISABLE_DEPRECATION_WARNINGS
 			if (CurrentConfig.Width != CurrentWidth || CurrentConfig.Height != CurrentHeight || CurrentConfig.MaxFramerate != CurrentFrameRate)
 			{
 				AMF_RESULT Result = AMF_OK;
@@ -400,90 +355,66 @@ namespace AVEncoder
 					UE_LOG(LogEncoderAMF, Error, TEXT("Amf failed to ReInit for config change"));
 				}
 			}
-			PRAGMA_ENABLE_DEPRECATION_WARNINGS
 
 			// Dynamic Properties
-			PRAGMA_DISABLE_DEPRECATION_WARNINGS
 			if (AmfEncoder->SetProperty(AMF_VIDEO_ENCODER_MIN_QP, FMath::Clamp<amf_int64>(CurrentConfig.QPMin, 0, 51)) != AMF_OK)
-			PRAGMA_ENABLE_DEPRECATION_WARNINGS
 			{
 				UE_LOG(LogEncoderAMF, Error, TEXT("Amf failed to set min qp"));
 			}
 
-			PRAGMA_DISABLE_DEPRECATION_WARNINGS
 			if (AmfEncoder->SetProperty(AMF_VIDEO_ENCODER_MAX_QP, CurrentConfig.QPMax > -1 ? FMath::Clamp<amf_int64>(CurrentConfig.QPMax, 0, 51) : 51) != AMF_OK)
-			PRAGMA_ENABLE_DEPRECATION_WARNINGS
 			{
 				UE_LOG(LogEncoderAMF, Error, TEXT("Amf failed to set max qp"));
 			}
 
-			PRAGMA_DISABLE_DEPRECATION_WARNINGS
 			if (AmfEncoder->SetProperty(AMF_VIDEO_ENCODER_TARGET_BITRATE, CurrentConfig.TargetBitrate) != AMF_OK)
-			PRAGMA_ENABLE_DEPRECATION_WARNINGS
 			{
 				UE_LOG(LogEncoderAMF, Error, TEXT("Amf failed to set target bitrate"));
 			}
 
 #if PLATFORM_WINDOWS
 			// Properties in this macro block are supposed to be dynamic but error when used with Vulkan
-			PRAGMA_DISABLE_DEPRECATION_WARNINGS
 			if (AmfEncoder->SetProperty(AMF_VIDEO_ENCODER_RATE_CONTROL_METHOD, ConvertRateControlModeAMF(CurrentConfig.RateControlMode)) != AMF_OK)
-			PRAGMA_ENABLE_DEPRECATION_WARNINGS
 			{
 				UE_LOG(LogEncoderAMF, Error, TEXT("Amf failed to set rate control method"));
 			}
 
-			PRAGMA_DISABLE_DEPRECATION_WARNINGS
 			if (CurrentConfig.RateControlMode == RateControlMode::CBR)
-			PRAGMA_ENABLE_DEPRECATION_WARNINGS
 			{
-				PRAGMA_DISABLE_DEPRECATION_WARNINGS
 				if (AmfEncoder->SetProperty(AMF_VIDEO_ENCODER_FILLER_DATA_ENABLE, CurrentConfig.FillData) != AMF_OK)
-				PRAGMA_ENABLE_DEPRECATION_WARNINGS
 				{
 					UE_LOG(LogEncoderAMF, Error, TEXT("Amf failed to enable filler data to maintain CBR"));
 				}
 			}
 
-			PRAGMA_DISABLE_DEPRECATION_WARNINGS
 			if (AmfEncoder->SetProperty(AMF_VIDEO_ENCODER_PEAK_BITRATE, CurrentConfig.MaxBitrate) != AMF_OK)
-			PRAGMA_ENABLE_DEPRECATION_WARNINGS
 			{
 				UE_LOG(LogEncoderAMF, Error, TEXT("Amf failed to set peak bitrate"));
 			}
 #endif
 
-			PRAGMA_DISABLE_DEPRECATION_WARNINGS
 			NeedsReconfigure = false;
-			PRAGMA_ENABLE_DEPRECATION_WARNINGS
 		}
-		
 	}
 
-	PRAGMA_DISABLE_DEPRECATION_WARNINGS
 	AMF_RESULT FVideoEncoderAmf_H264::FAMFLayer::Encode(const TSharedPtr<FVideoEncoderInputFrameImpl> frame, FEncodeOptions const& options)
-	PRAGMA_ENABLE_DEPRECATION_WARNINGS
 	{
 		AMF_RESULT Result = AMF_FAIL;
 		TSharedPtr<FInputOutput> Buffer = GetOrCreateSurface(frame);
 
 		if (Buffer)
 		{
-			PRAGMA_DISABLE_DEPRECATION_WARNINGS
 			if(CurrentConfig.Width != frame->GetWidth() || CurrentConfig.Height != frame->GetHeight())
 			{
 				CurrentConfig.Width = frame->GetWidth();
 				CurrentConfig.Height = frame->GetHeight();
 				NeedsReconfigure = true;
 			}
-			PRAGMA_ENABLE_DEPRECATION_WARNINGS
 
 			
 			MaybeReconfigure();
 
-			PRAGMA_DISABLE_DEPRECATION_WARNINGS
 			Buffer->Surface->SetPts(frame->GetTimestampRTP());
-			PRAGMA_ENABLE_DEPRECATION_WARNINGS
 			amf_int64 Start_ts = FPlatformTime::Cycles64();
 			Buffer->Surface->SetProperty(AMF_VIDEO_ENCODER_START_TS, Start_ts);
 			Buffer->Surface->SetProperty(AMF_BUFFER_INPUT_FRAME, uintptr_t(frame.Get()));
@@ -492,9 +423,7 @@ namespace AVEncoder
 			Buffer->Surface->SetProperty(AMF_VIDEO_ENCODER_STATISTICS_FEEDBACK, true);
 #endif
 
-			PRAGMA_DISABLE_DEPRECATION_WARNINGS
 			if (options.bForceKeyFrame)
-			PRAGMA_ENABLE_DEPRECATION_WARNINGS
 			{
 				if (Buffer->Surface->SetProperty(AMF_VIDEO_ENCODER_FORCE_PICTURE_TYPE, AMF_VIDEO_ENCODER_PICTURE_TYPE_IDR) != AMF_OK)
 				{
@@ -573,9 +502,7 @@ namespace AVEncoder
 			AMFBufferPtr OutBuffer(data);
 
 			// Create packet with buffer contents
-			PRAGMA_DISABLE_DEPRECATION_WARNINGS
 			FCodecPacket Packet = FCodecPacket::Create(static_cast<const uint8*>(OutBuffer->GetNative()), OutBuffer->GetSize());
-			PRAGMA_ENABLE_DEPRECATION_WARNINGS
 
 			uint32 PictureType = AMF_VIDEO_ENCODER_PICTURE_TYPE_NONE;
 			if (OutBuffer->GetProperty(AMF_VIDEO_ENCODER_OUTPUT_DATA_TYPE, &PictureType) != AMF_OK)
@@ -585,17 +512,12 @@ namespace AVEncoder
 			else if (PictureType == AMF_VIDEO_ENCODER_OUTPUT_DATA_TYPE_IDR)
 			{
 				UE_LOG(LogEncoderAMF, Verbose, TEXT("Generated IDR Frame"));
-				PRAGMA_DISABLE_DEPRECATION_WARNINGS
 				Packet.IsKeyFrame = true;
-				PRAGMA_ENABLE_DEPRECATION_WARNINGS
 			}
-			
 
 			if (RHIGetInterfaceType() != ERHIInterfaceType::Vulkan) // Amf with Vulkan doesn't currently support statistics
 			{
-				PRAGMA_DISABLE_DEPRECATION_WARNINGS
 				if (OutBuffer->GetProperty(AMF_VIDEO_ENCODER_STATISTIC_FRAME_QP, &Packet.VideoQP) != AMF_OK)
-				PRAGMA_ENABLE_DEPRECATION_WARNINGS
 				{
 					UE_LOG(LogEncoderAMF, Error, TEXT("Amf failed to get frame QP."));
 				}
@@ -607,11 +529,9 @@ namespace AVEncoder
 				UE_LOG(LogEncoderAMF, Error, TEXT("Amf failed to get encode start time."));
 			}
 
-			PRAGMA_DISABLE_DEPRECATION_WARNINGS
 			Packet.Timings.StartTs = FTimespan::FromSeconds(FPlatformTime::ToSeconds64(StartTs));
 			Packet.Timings.FinishTs = FTimespan::FromSeconds(FPlatformTime::ToSeconds64(FPlatformTime::Cycles64()));
 			Packet.Framerate = GetConfig().MaxFramerate;
-			PRAGMA_ENABLE_DEPRECATION_WARNINGS
 
 			TSharedPtr<FVideoEncoderInputFrameImpl> SourceFrame;
 			if (OutBuffer->GetProperty(AMF_BUFFER_INPUT_FRAME, (intptr_t*)&SourceFrame) != AMF_OK)
@@ -619,12 +539,10 @@ namespace AVEncoder
 				UE_LOG(LogEncoderAMF, Fatal, TEXT("Amf failed to get buffer input frame."));
 			}
 
-			PRAGMA_DISABLE_DEPRECATION_WARNINGS
 			if (Encoder.OnEncodedPacket)
 			{
 				Encoder.OnEncodedPacket(LayerIndex, SourceFrame, Packet);
 			}
-			PRAGMA_ENABLE_DEPRECATION_WARNINGS
 
 			bIsProcessingFrame = false;
 		}
@@ -649,21 +567,16 @@ namespace AVEncoder
 		return true;
 	}
 
-	PRAGMA_DISABLE_DEPRECATION_WARNINGS
 	amf::AMFVulkanSurface* CreateVulkanSurface(VkImage Image, VkDeviceMemory DeviceMemory, EVideoFrameFormat Format, uint32 Size, uint32 Width, uint32 Height)
-	PRAGMA_ENABLE_DEPRECATION_WARNINGS
 	{
 		amf_int32 pixelFormat = 0;
 		int pixelSize = 0;
-		PRAGMA_DISABLE_DEPRECATION_WARNINGS
-		switch (Format) 
-		{
-			default:
-			case EVideoFrameFormat::VULKAN_R8G8B8A8_UNORM:
-				pixelFormat = VK_FORMAT_B8G8R8A8_UNORM;
-				break;
+		switch (Format) {
+		default:
+		case EVideoFrameFormat::VULKAN_R8G8B8A8_UNORM:
+			pixelFormat = VK_FORMAT_B8G8R8A8_UNORM;
+			break;
 		}
-		PRAGMA_ENABLE_DEPRECATION_WARNINGS
 
 		amf::AMFVulkanSurface* SurfaceTexture = new amf::AMFVulkanSurface();
 		SurfaceTexture->cbSizeof = sizeof(amf::AMFVulkanSurface);
@@ -688,7 +601,6 @@ namespace AVEncoder
 	{
 		void* TextureToCompress = nullptr;
 
-		PRAGMA_DISABLE_DEPRECATION_WARNINGS
 		switch (InFrame->GetFormat())
 		{
 #if PLATFORM_WINDOWS
@@ -735,7 +647,6 @@ namespace AVEncoder
 			UE_LOG(LogEncoderAMF, Error, TEXT("Video Frame Format %s not supported by Amf on this platform."), *ToString(InFrame->GetFormat()));
 			break;
 		}
-		PRAGMA_ENABLE_DEPRECATION_WARNINGS
 
 		if (!TextureToCompress)
 		{
@@ -796,7 +707,6 @@ namespace AVEncoder
 
 		if (TextureToCompress)
 		{
-			PRAGMA_DISABLE_DEPRECATION_WARNINGS
 			switch (SourceFrame->GetFormat())
 			{
 #if PLATFORM_WINDOWS
@@ -815,7 +725,6 @@ namespace AVEncoder
 				UE_LOG(LogEncoderAMF, Error, TEXT("Video format %s not inplemented for Amf on this platform"), *ToString(SourceFrame->GetFormat()));
 				break;
 			}
-			PRAGMA_ENABLE_DEPRECATION_WARNINGS
 		}
 		else
 		{
@@ -862,33 +771,26 @@ namespace AVEncoder
 		return true;
 	}
 
-	PRAGMA_DISABLE_DEPRECATION_WARNINGS
 	static bool GetEncoderSupportedInputFormats(AMFIOCapsPtr IOCaps, TArray<EVideoFrameFormat>& OutSupportedInputFormats)
-	PRAGMA_ENABLE_DEPRECATION_WARNINGS
 	{
 		// TODO check if we actually need to query Amf for this
-		PRAGMA_DISABLE_DEPRECATION_WARNINGS
+
 #if PLATFORM_WINDOWS
 		OutSupportedInputFormats.Push(EVideoFrameFormat::D3D11_R8G8B8A8_UNORM);
 		OutSupportedInputFormats.Push(EVideoFrameFormat::D3D12_R8G8B8A8_UNORM);
 #endif
 		OutSupportedInputFormats.Push(EVideoFrameFormat::VULKAN_R8G8B8A8_UNORM);
-		PRAGMA_ENABLE_DEPRECATION_WARNINGS
 
 		return true;
 	}
 
 
-	PRAGMA_DISABLE_DEPRECATION_WARNINGS
 	static bool GetEncoderInfo(FAmfCommon& AMF, FVideoEncoderInfo& EncoderInfo)
-	PRAGMA_ENABLE_DEPRECATION_WARNINGS
 	{
 		bool bSuccess = true;
 
 		AMF.InitializeContext(RHIGetInterfaceType(), GDynamicRHI->GetName(), NULL);
-		PRAGMA_DISABLE_DEPRECATION_WARNINGS
 		EncoderInfo.CodecType = ECodecType::H264;
-		PRAGMA_ENABLE_DEPRECATION_WARNINGS
 		
 		// Create temp component
 		AMFComponentPtr TempEncoder;
@@ -909,10 +811,8 @@ namespace AVEncoder
 		uint32 LevelMax;
 		if (EncoderCaps->GetProperty(AMF_VIDEO_ENCODER_CAP_MAX_LEVEL, &LevelMax) == AMF_OK)
 		{
-			PRAGMA_DISABLE_DEPRECATION_WARNINGS
 			EncoderInfo.H264.MinLevel = 9;														// Like NVENC we hard min at 9
 			EncoderInfo.H264.MaxLevel = (LevelMax < 9) ? 9 : (LevelMax > 52) ? 52 : LevelMax;	// Like NVENC we hard max at 52
-			PRAGMA_ENABLE_DEPRECATION_WARNINGS
 		}
 		else
 		{
@@ -920,10 +820,8 @@ namespace AVEncoder
 			bSuccess = false;
 		}
 
-		PRAGMA_DISABLE_DEPRECATION_WARNINGS
 		if (!GetEncoderSupportedProfiles(EncoderCaps, EncoderInfo.H264.SupportedProfiles) ||
 			!GetEncoderSupportedInputFormats(InputCaps, EncoderInfo.SupportedInputFormats))
-		PRAGMA_ENABLE_DEPRECATION_WARNINGS
 		{
 			bSuccess = false;
 		}
diff --git a/Engine/Plugins/Media/HardwareEncoders/Source/EncoderAMF/Private/Amf_EncoderH264.h b/Engine/Plugins/Media/HardwareEncoders/Source/EncoderAMF/Private/Amf_EncoderH264.h
--- a/Engine/Plugins/Media/HardwareEncoders/Source/EncoderAMF/Private/Amf_EncoderH264.h
+++ b/Engine/Plugins/Media/HardwareEncoders/Source/EncoderAMF/Private/Amf_EncoderH264.h
@@ -24,40 +24,30 @@ namespace AVEncoder
 {
 	using namespace amf;
 
-    PRAGMA_DISABLE_DEPRECATION_WARNINGS
-	class FVideoEncoderAmf_H264 : public FVideoEncoder
-    PRAGMA_ENABLE_DEPRECATION_WARNINGS
-	{
+    class FVideoEncoderAmf_H264 : public FVideoEncoder
+    {
     public:
         virtual ~FVideoEncoderAmf_H264() override;
 
         // query whether or not encoder is supported and available
-        PRAGMA_DISABLE_DEPRECATION_WARNINGS
-		static bool GetIsAvailable(FVideoEncoderInputImpl &InInput, FVideoEncoderInfo &OutEncoderInfo);
-		PRAGMA_ENABLE_DEPRECATION_WARNINGS
+        static bool GetIsAvailable(FVideoEncoderInputImpl &InInput, FVideoEncoderInfo &OutEncoderInfo);
 
         // register encoder with video encoder factory
         static void Register(FVideoEncoderFactory &InFactory);
 
-        PRAGMA_DISABLE_DEPRECATION_WARNINGS
-		bool Setup(TSharedRef<FVideoEncoderInput> input, FLayerConfig const& config) override;
-		void Encode(const TSharedPtr<FVideoEncoderInputFrame> frame, FEncodeOptions const& options) override;
-        PRAGMA_ENABLE_DEPRECATION_WARNINGS
+        bool Setup(TSharedRef<FVideoEncoderInput> input, FLayerConfig const& config) override;
+        void Encode(const TSharedPtr<FVideoEncoderInputFrame> frame, FEncodeOptions const& options) override;
         void Flush();
         void Shutdown() override;
 
     protected:
-        PRAGMA_DISABLE_DEPRECATION_WARNINGS
-		FLayer* CreateLayer(uint32 layerIdx, FLayerConfig const& config) override;
-		void DestroyLayer(FLayer *layer) override;
-        PRAGMA_ENABLE_DEPRECATION_WARNINGS
+        FLayer* CreateLayer(uint32 layerIdx, FLayerConfig const& config) override;
+        void DestroyLayer(FLayer *layer) override;
 
     private:
         FVideoEncoderAmf_H264();
 
-		PRAGMA_DISABLE_DEPRECATION_WARNINGS
 		class FAMFLayer : public FLayer
-		PRAGMA_ENABLE_DEPRECATION_WARNINGS
 		{
 		public:
 
@@ -68,9 +58,7 @@ namespace AVEncoder
 					AMFSurfacePtr Surface;
 			};
 
-			PRAGMA_DISABLE_DEPRECATION_WARNINGS
 			FAMFLayer(uint32 layerIdx, FLayerConfig const& config, FVideoEncoderAmf_H264& encoder);
-			PRAGMA_ENABLE_DEPRECATION_WARNINGS
 			~FAMFLayer();
 
 			bool Setup();
@@ -84,9 +72,7 @@ namespace AVEncoder
 			template<class T>
 			bool GetCapability(const TCHAR* CapToQuery, T& OutCap) const;
 
-			PRAGMA_DISABLE_DEPRECATION_WARNINGS
 			AMF_RESULT Encode(const TSharedPtr<FVideoEncoderInputFrameImpl> frame, FEncodeOptions const& options);
-			PRAGMA_ENABLE_DEPRECATION_WARNINGS
 			void Flush();
 			void Shutdown();
 			void UpdateBitrate(uint32 InMaxBitRate, uint32 InTargetBitRate);
@@ -111,9 +97,7 @@ namespace AVEncoder
 		};
 
 		FAmfCommon& Amf;
-		PRAGMA_DISABLE_DEPRECATION_WARNINGS
 		EVideoFrameFormat FrameFormat = EVideoFrameFormat::Undefined;
-		PRAGMA_ENABLE_DEPRECATION_WARNINGS
 		void* EncoderDevice;
 
 		uint32 MaxFramerate = 0;
diff --git a/Engine/Plugins/Media/HardwareEncoders/Source/EncoderAMF/Private/Amf_EncoderModule.cpp b/Engine/Plugins/Media/HardwareEncoders/Source/EncoderAMF/Private/Amf_EncoderModule.cpp
--- a/Engine/Plugins/Media/HardwareEncoders/Source/EncoderAMF/Private/Amf_EncoderModule.cpp
+++ b/Engine/Plugins/Media/HardwareEncoders/Source/EncoderAMF/Private/Amf_EncoderModule.cpp
@@ -55,11 +55,8 @@ public:
 				{
 					if (IsRHIDeviceAMD())
 					{
-						PRAGMA_DISABLE_DEPRECATION_WARNINGS
 						FVideoEncoderAmf_H264::Register(FVideoEncoderFactory::Get());
-						PRAGMA_ENABLE_DEPRECATION_WARNINGS
 					}
-					
 				});
 				
 				AMFStarted = true;
diff --git a/Engine/Plugins/Media/HardwareEncoders/Source/EncoderNVENC/Private/NVENC_EncoderH264.cpp b/Engine/Plugins/Media/HardwareEncoders/Source/EncoderNVENC/Private/NVENC_EncoderH264.cpp
--- a/Engine/Plugins/Media/HardwareEncoders/Source/EncoderNVENC/Private/NVENC_EncoderH264.cpp
+++ b/Engine/Plugins/Media/HardwareEncoders/Source/EncoderNVENC/Private/NVENC_EncoderH264.cpp
@@ -17,7 +17,6 @@
 
 namespace
 {
-	PRAGMA_DISABLE_DEPRECATION_WARNINGS
 	NV_ENC_PARAMS_RC_MODE ConvertRateControlModeNVENC(AVEncoder::FVideoEncoder::RateControlMode mode)
 	{
 		switch (mode)
@@ -69,7 +68,6 @@ namespace
 				return NV_ENC_H264_PROFILE_CONSTRAINED_HIGH_GUID;
 		}
 	}
-	PRAGMA_ENABLE_DEPRECATION_WARNINGS
 } // namespace
 
 namespace AVEncoder
@@ -164,63 +162,43 @@ namespace AVEncoder
 		});
 	}
 
-	PRAGMA_DISABLE_DEPRECATION_WARNINGS
 	static bool GetEncoderInfo(FNVENCCommon& NVENC, FVideoEncoderInfo& EncoderInfo);
-	PRAGMA_ENABLE_DEPRECATION_WARNINGS
 	static int GetEncoderCapability(FNVENCCommon& NVENC, void* InEncoder, NV_ENC_CAPS InCapsToQuery);
 
-	PRAGMA_DISABLE_DEPRECATION_WARNINGS
 	bool FVideoEncoderNVENC_H264::GetIsAvailable(const FVideoEncoderInput& InVideoFrameFactory, FVideoEncoderInfo& OutEncoderInfo)
-	PRAGMA_ENABLE_DEPRECATION_WARNINGS
 	{
 		FNVENCCommon& NVENC = FNVENCCommon::Setup();
 		bool bIsAvailable = NVENC.GetIsAvailable();
 		if (bIsAvailable)
 		{
-			PRAGMA_DISABLE_DEPRECATION_WARNINGS
 			OutEncoderInfo.CodecType = ECodecType::H264;
-			PRAGMA_ENABLE_DEPRECATION_WARNINGS
 		}
-		
 		return bIsAvailable;
 	}
 
-	PRAGMA_DISABLE_DEPRECATION_WARNINGS
 	void FVideoEncoderNVENC_H264::Register(FVideoEncoderFactory& InFactory)
-	PRAGMA_ENABLE_DEPRECATION_WARNINGS
 	{
 		FNVENCCommon& NVENC = FNVENCCommon::Setup();
 		if (NVENC.GetIsAvailable() && IsRHIDeviceNVIDIA())
 		{
-			PRAGMA_DISABLE_DEPRECATION_WARNINGS
 			FVideoEncoderInfo EncoderInfo;
-			PRAGMA_ENABLE_DEPRECATION_WARNINGS
 			if (GetEncoderInfo(NVENC, EncoderInfo))
 			{
-				PRAGMA_DISABLE_DEPRECATION_WARNINGS
 				InFactory.Register(EncoderInfo, []() { return TUniquePtr<FVideoEncoder>(new FVideoEncoderNVENC_H264()); });
-				PRAGMA_ENABLE_DEPRECATION_WARNINGS
 			}
-			
 		}
 	}
 
-	PRAGMA_DISABLE_DEPRECATION_WARNINGS
 	FVideoEncoderNVENC_H264::FVideoEncoderNVENC_H264()
 		: NVENC(FNVENCCommon::Setup())
 	{
 		// Parse NVENC settings from command line (if any relevant ones are passed)
 		NVENCParseCommandLineFlags();
 	}
-	PRAGMA_ENABLE_DEPRECATION_WARNINGS
 
-	PRAGMA_DISABLE_DEPRECATION_WARNINGS
 	FVideoEncoderNVENC_H264::~FVideoEncoderNVENC_H264() { Shutdown(); }
-	PRAGMA_ENABLE_DEPRECATION_WARNINGS
 
-	PRAGMA_DISABLE_DEPRECATION_WARNINGS
 	bool FVideoEncoderNVENC_H264::Setup(TSharedRef<FVideoEncoderInput> InputFrameFactory, const FLayerConfig& InitConfig)
-	PRAGMA_ENABLE_DEPRECATION_WARNINGS
 	{
 		if (!NVENC.GetIsAvailable())
 		{
@@ -228,7 +206,6 @@ namespace AVEncoder
 			return false;
 		}
 
-		PRAGMA_DISABLE_DEPRECATION_WARNINGS
 		FrameFormat = InputFrameFactory->GetFrameFormat();
 		switch (FrameFormat)
 		{
@@ -246,7 +223,6 @@ namespace AVEncoder
 				UE_LOG(LogEncoderNVENC, Error, TEXT("Frame format %s is not supported by NVENC_Encoder on this platform."), *ToString(FrameFormat));
 				return false;
 		}
-		PRAGMA_ENABLE_DEPRECATION_WARNINGS
 
 		if (!EncoderDevice)
 		{
@@ -254,7 +230,6 @@ namespace AVEncoder
 			return false;
 		}
 
-		PRAGMA_DISABLE_DEPRECATION_WARNINGS
 		FLayerConfig mutableConfig = InitConfig;
 		if (mutableConfig.MaxFramerate == 0)
 		{
@@ -262,12 +237,9 @@ namespace AVEncoder
 		}
 
 		return AddLayer(mutableConfig);
-		PRAGMA_ENABLE_DEPRECATION_WARNINGS
 	}
 
-	PRAGMA_DISABLE_DEPRECATION_WARNINGS
 	FVideoEncoder::FLayer* FVideoEncoderNVENC_H264::CreateLayer(uint32 layerIdx, FLayerConfig const& config)
-	PRAGMA_ENABLE_DEPRECATION_WARNINGS
 	{
 		auto const layer = new FNVENCLayer(layerIdx, config, *this);
 		if (!layer->Setup())
@@ -278,17 +250,11 @@ namespace AVEncoder
 		return layer;
 	}
 
-	PRAGMA_DISABLE_DEPRECATION_WARNINGS
 	void FVideoEncoderNVENC_H264::DestroyLayer(FLayer* layer) { delete layer; }
-	PRAGMA_ENABLE_DEPRECATION_WARNINGS
 
-	PRAGMA_DISABLE_DEPRECATION_WARNINGS
 	void FVideoEncoderNVENC_H264::Encode(const TSharedPtr<FVideoEncoderInputFrame> InFrame, const FEncodeOptions& EncodeOptions)
-	PRAGMA_ENABLE_DEPRECATION_WARNINGS
 	{
-		PRAGMA_DISABLE_DEPRECATION_WARNINGS
 		for (FVideoEncoder::FLayer* Layer : Layers)
-		PRAGMA_ENABLE_DEPRECATION_WARNINGS
 		{
 			FVideoEncoderNVENC_H264::FNVENCLayer* NVLayer = static_cast<FVideoEncoderNVENC_H264::FNVENCLayer*>(Layer);
 			NVLayer->Encode(InFrame, EncodeOptions);
@@ -297,9 +263,7 @@ namespace AVEncoder
 
 	void FVideoEncoderNVENC_H264::Flush()
 	{
-		PRAGMA_DISABLE_DEPRECATION_WARNINGS
 		for (FVideoEncoder::FLayer* Layer : Layers)
-		PRAGMA_ENABLE_DEPRECATION_WARNINGS
 		{
 			FVideoEncoderNVENC_H264::FNVENCLayer* NVLayer = static_cast<FVideoEncoderNVENC_H264::FNVENCLayer*>(Layer);
 			NVLayer->Flush();
@@ -308,22 +272,16 @@ namespace AVEncoder
 
 	void FVideoEncoderNVENC_H264::Shutdown()
 	{
-		PRAGMA_DISABLE_DEPRECATION_WARNINGS
 		for (FVideoEncoder::FLayer* Layer : Layers)
-		PRAGMA_ENABLE_DEPRECATION_WARNINGS
 		{
 			FVideoEncoderNVENC_H264::FNVENCLayer* NVLayer = static_cast<FVideoEncoderNVENC_H264::FNVENCLayer*>(Layer);
 			NVLayer->Shutdown();
 			DestroyLayer(NVLayer);
 		}
-		PRAGMA_DISABLE_DEPRECATION_WARNINGS
 		Layers.Reset();
-		PRAGMA_ENABLE_DEPRECATION_WARNINGS
 	}
-	
 
 	// --- FVideoEncoderNVENC_H264::FLayer ------------------------------------------------------------
-	PRAGMA_DISABLE_DEPRECATION_WARNINGS
 	FVideoEncoderNVENC_H264::FNVENCLayer::FNVENCLayer(uint32 layerIdx, FLayerConfig const& config, FVideoEncoderNVENC_H264& encoder)
 		: FLayer(config)
 		, Encoder(encoder)
@@ -332,11 +290,8 @@ namespace AVEncoder
 		, LayerIndex(layerIdx)
 	{
 	}
-	PRAGMA_ENABLE_DEPRECATION_WARNINGS
 
-	PRAGMA_DISABLE_DEPRECATION_WARNINGS
 	FVideoEncoderNVENC_H264::FNVENCLayer::~FNVENCLayer() {}
-	PRAGMA_ENABLE_DEPRECATION_WARNINGS
 
 	bool FVideoEncoderNVENC_H264::FNVENCLayer::Setup()
 	{
@@ -367,7 +322,6 @@ namespace AVEncoder
 			OpenEncodeSessionExParams.apiVersion = NVENCAPI_VERSION;
 			OpenEncodeSessionExParams.device = Encoder.EncoderDevice;
 
-			PRAGMA_DISABLE_DEPRECATION_WARNINGS
 			switch (Encoder.FrameFormat)
 			{
 				case EVideoFrameFormat::D3D11_R8G8B8A8_UNORM:
@@ -382,7 +336,6 @@ namespace AVEncoder
 					return false;
 					break;
 			}
-			PRAGMA_ENABLE_DEPRECATION_WARNINGS
 
 			auto const result = NVENC.nvEncOpenEncodeSessionEx(&OpenEncodeSessionExParams, &NVEncoder);
 			if (result != NV_ENC_SUCCESS)
@@ -401,16 +354,12 @@ namespace AVEncoder
 		// set the initialization parameters
 		FMemory::Memzero(EncoderInitParams);
 
-		PRAGMA_DISABLE_DEPRECATION_WARNINGS
 		CurrentConfig.MaxFramerate = 60;
-		PRAGMA_ENABLE_DEPRECATION_WARNINGS
 
 		EncoderInitParams.version = NV_ENC_INITIALIZE_PARAMS_VER;
 		EncoderInitParams.encodeGUID = NV_ENC_CODEC_H264_GUID;
 		EncoderInitParams.presetGUID = NV_ENC_PRESET_P4_GUID;
-		PRAGMA_DISABLE_DEPRECATION_WARNINGS
 		EncoderInitParams.frameRateNum = CurrentConfig.MaxFramerate;
-		PRAGMA_ENABLE_DEPRECATION_WARNINGS
 		EncoderInitParams.frameRateDen = 1;
 		EncoderInitParams.enablePTD = 1;
 		EncoderInitParams.reportSliceOffsets = 0;
@@ -431,9 +380,7 @@ namespace AVEncoder
 
 		// copy the preset config to our config
 		FMemory::Memcpy(&EncoderConfig, &PresetConfig.presetCfg, sizeof(NV_ENC_CONFIG));
-		PRAGMA_DISABLE_DEPRECATION_WARNINGS
 		EncoderConfig.profileGUID = ConvertH264Profile(CurrentConfig.H264Profile);
-		PRAGMA_ENABLE_DEPRECATION_WARNINGS
 		EncoderConfig.rcParams.version = NV_ENC_RC_PARAMS_VER;
 
 		EncoderInitParams.encodeConfig = &EncoderConfig;
@@ -500,7 +447,6 @@ namespace AVEncoder
 	{
 		uint64 PreReconfigureCycles = FPlatformTime::Cycles64();
 
-		PRAGMA_DISABLE_DEPRECATION_WARNINGS
 		FScopeLock lock(&ConfigMutex);
 		if (NeedsReconfigure)
 		{
@@ -538,7 +484,6 @@ namespace AVEncoder
 				FNVENCStats::Get().SetReconfigureLatency(ReconfigureDeltaMs);
 			}
 		}
-		PRAGMA_ENABLE_DEPRECATION_WARNINGS
 	}
 
 	void FVideoEncoderNVENC_H264::FNVENCLayer::UpdateLastEncodedQP(uint32 InLastEncodedQP)
@@ -550,15 +495,11 @@ namespace AVEncoder
 		}
 
 		LastEncodedQP = InLastEncodedQP;
-		PRAGMA_DISABLE_DEPRECATION_WARNINGS
 		NeedsReconfigure = true;
-		PRAGMA_ENABLE_DEPRECATION_WARNINGS
 	}
-	
 
 	void FVideoEncoderNVENC_H264::FNVENCLayer::UpdateConfig()
 	{
-		PRAGMA_DISABLE_DEPRECATION_WARNINGS
 		EncoderInitParams.encodeWidth = EncoderInitParams.darWidth = CurrentConfig.Width;
 		EncoderInitParams.encodeHeight = EncoderInitParams.darHeight = CurrentConfig.Height;
 
@@ -574,7 +515,6 @@ namespace AVEncoder
 		RateControlParams.maxQP = { MaxQP, MaxQP, MaxQP };
 		RateControlParams.enableMinQP = CurrentConfig.QPMin > -1;
 		RateControlParams.enableMaxQP = CurrentConfig.QPMax > -1;
-		PRAGMA_ENABLE_DEPRECATION_WARNINGS
 
 		// If we have QP ranges turned on use the last encoded QP to guide the max QP for an i-frame, so the i-frame doesn't look too blocky
 		// Note: this does nothing if we have i-frames turned off.
@@ -583,16 +523,12 @@ namespace AVEncoder
 			RateControlParams.maxQP.qpIntra = LastEncodedQP;
 		}
 
-		PRAGMA_DISABLE_DEPRECATION_WARNINGS
 		EncoderInitParams.encodeConfig->profileGUID = ConvertH264Profile(CurrentConfig.H264Profile);
-		PRAGMA_ENABLE_DEPRECATION_WARNINGS
-		
+
 		NV_ENC_CONFIG_H264& H264Config = EncoderInitParams.encodeConfig->encodeCodecConfig.h264Config;
-		PRAGMA_DISABLE_DEPRECATION_WARNINGS
 		H264Config.enableFillerDataInsertion = CurrentConfig.FillData ? 1 : 0;
 
 		if (CurrentConfig.RateControlMode == AVEncoder::FVideoEncoder::RateControlMode::CBR && CurrentConfig.FillData)
-		PRAGMA_ENABLE_DEPRECATION_WARNINGS
 		{
 			// `outputPictureTimingSEI` is used in CBR mode to fill video frame with data to match the requested bitrate.
 			H264Config.outputPictureTimingSEI = 1;
@@ -683,9 +619,7 @@ namespace AVEncoder
 		}
 	}
 
-	PRAGMA_DISABLE_DEPRECATION_WARNINGS
 	void FVideoEncoderNVENC_H264::FNVENCLayer::Encode(const TSharedPtr<FVideoEncoderInputFrame> InFrame, const FEncodeOptions& EncodeOptions)
-	PRAGMA_ENABLE_DEPRECATION_WARNINGS
 	{
 		// This encode is single threaded and blocking, if someone else is encoding right now, we wait.
 		while (bIsProcessingFrame) 
@@ -710,9 +644,7 @@ namespace AVEncoder
 		Buffer->PicParams = {};
 		Buffer->PicParams.encodePicFlags = 0;
 
-		PRAGMA_DISABLE_DEPRECATION_WARNINGS
 		if (EncodeOptions.bForceKeyFrame)
-		PRAGMA_ENABLE_DEPRECATION_WARNINGS
 		{
 			LastKeyFrameTime = FDateTime::UtcNow();
 			Buffer->PicParams.encodePicFlags |= NV_ENC_PIC_FLAG_FORCEIDR;
@@ -721,10 +653,8 @@ namespace AVEncoder
 		// Send SPS/PPS every frame (wasteful)
 		//Buffer->PicParams.encodePicFlags |= NV_ENC_PIC_FLAG_OUTPUT_SPSPPS;
 
-		PRAGMA_DISABLE_DEPRECATION_WARNINGS
 		Buffer->PicParams.frameIdx = InFrame->GetFrameID();
 		Buffer->PicParams.inputTimeStamp = Buffer->TimeStamp = InFrame->GetTimestampUs();
-		PRAGMA_ENABLE_DEPRECATION_WARNINGS
 
 		EncodeBuffer(Buffer);
 
@@ -783,22 +713,15 @@ namespace AVEncoder
 		// lock output buffers for CPU access
 		if (LockOutputBuffer(Buffer))
 		{
-			PRAGMA_DISABLE_DEPRECATION_WARNINGS
 			if (Encoder.OnEncodedPacket)
-			PRAGMA_ENABLE_DEPRECATION_WARNINGS
 			{
-				PRAGMA_DISABLE_DEPRECATION_WARNINGS
 				FCodecPacket Packet = FCodecPacket::Create(static_cast<const uint8*>(Buffer->BitstreamData), Buffer->BitstreamDataSize);
-				PRAGMA_ENABLE_DEPRECATION_WARNINGS
 
 				if (Buffer->PictureType & NV_ENC_PIC_TYPE_IDR)
 				{
 					UE_LOG(LogEncoderNVENC, Verbose, TEXT("Generated IDR Frame"));
-					PRAGMA_DISABLE_DEPRECATION_WARNINGS
 					Packet.IsKeyFrame = true;
-					PRAGMA_ENABLE_DEPRECATION_WARNINGS
 				}
-				
 				else
 				{
 					// If it is not a keyframe store the QP.
@@ -806,12 +729,10 @@ namespace AVEncoder
 				}
 
 				FinishedEncodingCycles = FPlatformTime::Cycles64();
-				PRAGMA_DISABLE_DEPRECATION_WARNINGS
 				Packet.VideoQP = Buffer->FrameAvgQP;
 				Packet.Timings.StartTs = FTimespan::FromMilliseconds(FPlatformTime::ToMilliseconds64(Buffer->SubmitTimeCycles));
 				Packet.Timings.FinishTs = FTimespan::FromMilliseconds(FPlatformTime::ToMilliseconds64(FinishedEncodingCycles));
 				Packet.Framerate = EncoderInitParams.frameRateNum;
-				PRAGMA_ENABLE_DEPRECATION_WARNINGS
 
 				// Stats
 				if (CVarNVENCEnableStats.GetValueOnAnyThread())
@@ -821,19 +742,15 @@ namespace AVEncoder
 					FNVENCStats::Get().SetProcessFramesFuncLatency(ProcessFramesFuncDeltaMs);
 
 					// Total time the encoder took to process a frame from submit through to encode
-					PRAGMA_DISABLE_DEPRECATION_WARNINGS
 					double TotalEncoderLatencyMs = Packet.Timings.FinishTs.GetTotalMilliseconds() - Packet.Timings.StartTs.GetTotalMilliseconds();
-					PRAGMA_ENABLE_DEPRECATION_WARNINGS
 					FNVENCStats::Get().SetTotalEncoderLatency(TotalEncoderLatencyMs);
 				}
 
 				// Send frame to OnEncodedPacket listener, such as WebRTC.
-				PRAGMA_DISABLE_DEPRECATION_WARNINGS
 				if (Encoder.OnEncodedPacket)
 				{
 					Encoder.OnEncodedPacket(LayerIndex, Buffer->SourceFrame, Packet);
 				}
-				PRAGMA_ENABLE_DEPRECATION_WARNINGS
 			}
 
 			if (!UnlockOutputBuffer(Buffer))
@@ -873,24 +790,17 @@ namespace AVEncoder
 		else
 		{
 			// Check for buffer and Frame texture resolution mismatch
-			PRAGMA_DISABLE_DEPRECATION_WARNINGS
 			if (InFrame->GetWidth() != InputOutputBuffer->Width || InFrame->GetHeight() != InputOutputBuffer->Height)
-			PRAGMA_ENABLE_DEPRECATION_WARNINGS
 			{
 				DestroyBuffer(InputOutputBuffer);
 				InputOutputBuffer = CreateBuffer();
-				PRAGMA_DISABLE_DEPRECATION_WARNINGS
 				NeedsReconfigure = true;
-				PRAGMA_ENABLE_DEPRECATION_WARNINGS
 			}
-			
 		}
 
 		InputOutputBuffer->SourceFrame = InFrame;
-		PRAGMA_DISABLE_DEPRECATION_WARNINGS
 		InputOutputBuffer->Width = InFrame->GetWidth();
 		InputOutputBuffer->Height = InFrame->GetHeight();
-		PRAGMA_ENABLE_DEPRECATION_WARNINGS
 		return InputOutputBuffer;
 	}
 
@@ -1000,7 +910,6 @@ namespace AVEncoder
 			return;
 		}
 
-		PRAGMA_DISABLE_DEPRECATION_WARNINGS
 		switch (InBuffer->SourceFrame->GetCUDA().UnderlyingRHI)
 		{
 			case FVideoEncoderInputFrame::EUnderlyingRHI::Vulkan:
@@ -1013,7 +922,6 @@ namespace AVEncoder
 			default:
 				break;
 		}
-		PRAGMA_ENABLE_DEPRECATION_WARNINGS
 
 		InBuffer->Width = TextureSize.X;
 		InBuffer->Height = TextureSize.Y;
@@ -1041,7 +949,6 @@ namespace AVEncoder
 
 			void* TextureToCompress = nullptr;
 
-			PRAGMA_DISABLE_DEPRECATION_WARNINGS
 			switch (InBuffer->SourceFrame->GetFormat())
 			{
 #if PLATFORM_WINDOWS
@@ -1059,7 +966,6 @@ namespace AVEncoder
 				default:
 					break;
 			}
-			PRAGMA_ENABLE_DEPRECATION_WARNINGS
 
 			if (!TextureToCompress)
 			{
@@ -1070,7 +976,6 @@ namespace AVEncoder
 			InBuffer->InputTexture = TextureToCompress;
 			NVENCStruct(NV_ENC_REGISTER_RESOURCE, RegisterParam);
 
-			PRAGMA_DISABLE_DEPRECATION_WARNINGS
 			FIntPoint TextureSize(InBuffer->SourceFrame->GetWidth(), InBuffer->SourceFrame->GetHeight());
 
 			switch (InBuffer->SourceFrame->GetFormat())
@@ -1088,7 +993,6 @@ namespace AVEncoder
 				default:
 					break;
 			}
-			PRAGMA_ENABLE_DEPRECATION_WARNINGS
 
 			RegisterParam.resourceToRegister = InBuffer->InputTexture;
 
@@ -1403,9 +1307,7 @@ namespace AVEncoder
 		return OutSupportedProfiles != 0;
 	}
 
-	PRAGMA_DISABLE_DEPRECATION_WARNINGS
 	static bool GetEncoderSupportedInputFormats(FNVENCCommon& NVENC, void* InEncoder, TArray<EVideoFrameFormat>& OutSupportedInputFormats)
-	PRAGMA_ENABLE_DEPRECATION_WARNINGS
 	{
 		const uint32_t MaxInputFmtCount = 32;
 		uint32_t InputFmtCount = 0;
@@ -1428,22 +1330,18 @@ namespace AVEncoder
 				case NV_ENC_BUFFER_FORMAT_ARGB:
 					break;
 				case NV_ENC_BUFFER_FORMAT_ABGR:
-				PRAGMA_DISABLE_DEPRECATION_WARNINGS
 #if PLATFORM_WINDOWS
 					OutSupportedInputFormats.Push(EVideoFrameFormat::D3D11_R8G8B8A8_UNORM);
 					OutSupportedInputFormats.Push(EVideoFrameFormat::D3D12_R8G8B8A8_UNORM);
 #endif
 					OutSupportedInputFormats.Push(EVideoFrameFormat::CUDA_R8G8B8A8_UNORM);
-				PRAGMA_ENABLE_DEPRECATION_WARNINGS
 					break;
 			}
 		}
 		return true;
 	}
 
-	PRAGMA_DISABLE_DEPRECATION_WARNINGS
 	static bool GetEncoderInfo(FNVENCCommon& NVENC, FVideoEncoderInfo& EncoderInfo)
-	PRAGMA_ENABLE_DEPRECATION_WARNINGS
 	{
 		bool bSuccess = true;
 
@@ -1486,21 +1384,17 @@ namespace AVEncoder
 			return false;
 		}
 
-		PRAGMA_DISABLE_DEPRECATION_WARNINGS
 		EncoderInfo.CodecType = ECodecType::H264;
 		EncoderInfo.MaxWidth = GetEncoderCapability(NVENC, EncoderSession, NV_ENC_CAPS_WIDTH_MAX);
 		EncoderInfo.MaxHeight = GetEncoderCapability(NVENC, EncoderSession, NV_ENC_CAPS_HEIGHT_MAX);
-		PRAGMA_ENABLE_DEPRECATION_WARNINGS
 
 		int LevelMax = GetEncoderCapability(NVENC, EncoderSession, NV_ENC_CAPS_LEVEL_MAX);
 		int LevelMin = GetEncoderCapability(NVENC, EncoderSession, NV_ENC_CAPS_LEVEL_MIN);
 		if (LevelMin > 0 && LevelMax > 0 && LevelMax >= LevelMin)
 		{
-			PRAGMA_DISABLE_DEPRECATION_WARNINGS
 			EncoderInfo.H264.MinLevel = (LevelMin > 9) ? LevelMin : 9;
 			EncoderInfo.H264.MaxLevel = (LevelMax < 9) ? 9 : (LevelMax > NV_ENC_LEVEL_H264_52) ? NV_ENC_LEVEL_H264_52
 																							   : LevelMax;
-			PRAGMA_ENABLE_DEPRECATION_WARNINGS
 		}
 		else
 		{
@@ -1508,9 +1402,7 @@ namespace AVEncoder
 			bSuccess = false;
 		}
 
-		PRAGMA_DISABLE_DEPRECATION_WARNINGS
 		if (!GetEncoderSupportedProfiles(NVENC, EncoderSession, EncoderInfo.H264.SupportedProfiles) || !GetEncoderSupportedInputFormats(NVENC, EncoderSession, EncoderInfo.SupportedInputFormats))
-		PRAGMA_ENABLE_DEPRECATION_WARNINGS
 		{
 			bSuccess = false;
 		}
diff --git a/Engine/Plugins/Media/HardwareEncoders/Source/EncoderNVENC/Private/NVENC_EncoderH264.h b/Engine/Plugins/Media/HardwareEncoders/Source/EncoderNVENC/Private/NVENC_EncoderH264.h
--- a/Engine/Plugins/Media/HardwareEncoders/Source/EncoderNVENC/Private/NVENC_EncoderH264.h
+++ b/Engine/Plugins/Media/HardwareEncoders/Source/EncoderNVENC/Private/NVENC_EncoderH264.h
@@ -14,44 +14,32 @@ namespace AVEncoder { class FVideoEncoderInputFrameImpl; }
 
 namespace AVEncoder
 {
-	PRAGMA_DISABLE_DEPRECATION_WARNINGS
 	class FVideoEncoderNVENC_H264 : public FVideoEncoder
-	PRAGMA_ENABLE_DEPRECATION_WARNINGS
 	{
 
 	public:
 		virtual ~FVideoEncoderNVENC_H264() override; 
 
-		PRAGMA_DISABLE_DEPRECATION_WARNINGS
 		bool Setup(TSharedRef<FVideoEncoderInput> InputFrameFactory, const FLayerConfig& InitConfig) override;
-		PRAGMA_ENABLE_DEPRECATION_WARNINGS
 		void Shutdown() override;
 
 		// Query whether or not encoder is supported and available
-		PRAGMA_DISABLE_DEPRECATION_WARNINGS
 		static bool GetIsAvailable(const FVideoEncoderInput& InFrameFactory, FVideoEncoderInfo& OutEncoderInfo);
-		PRAGMA_ENABLE_DEPRECATION_WARNINGS
 
 		// Register encoder with video encoder factory
 		static void Register(FVideoEncoderFactory& InFactory);
 
-		PRAGMA_DISABLE_DEPRECATION_WARNINGS
 		void Encode(const TSharedPtr<FVideoEncoderInputFrame> InFrame, const FEncodeOptions& EncodeOptions) override;
-		PRAGMA_ENABLE_DEPRECATION_WARNINGS
 		void Flush();
 
 	protected:
-		PRAGMA_DISABLE_DEPRECATION_WARNINGS
 		FLayer* CreateLayer(uint32 InLayerIndex, const FLayerConfig& InLayerConfig) override;
 		void DestroyLayer(FLayer* layer) override;
-		PRAGMA_ENABLE_DEPRECATION_WARNINGS
 
 	private:
 		FVideoEncoderNVENC_H264();
 
-		PRAGMA_DISABLE_DEPRECATION_WARNINGS
 		class FNVENCLayer : public FLayer
-		PRAGMA_ENABLE_DEPRECATION_WARNINGS
 		{
 
 		public:
@@ -75,9 +63,7 @@ namespace AVEncoder
 				uint64 SubmitTimeCycles;
 			};
 
-			PRAGMA_DISABLE_DEPRECATION_WARNINGS
 			FNVENCLayer(uint32 layerIdx, FLayerConfig const& config, FVideoEncoderNVENC_H264& encoder);
-			PRAGMA_ENABLE_DEPRECATION_WARNINGS
 			~FNVENCLayer();
 			bool Setup();
 			bool CreateSession();
@@ -87,9 +73,7 @@ namespace AVEncoder
 			void MaybeReconfigure();
 			void UpdateConfig();
 			void UpdateLastEncodedQP(uint32 InLastEncodedQP);
-			PRAGMA_DISABLE_DEPRECATION_WARNINGS
 			void Encode(const TSharedPtr<FVideoEncoderInputFrame> InFrame, const FEncodeOptions& EncodeOptions);
-			PRAGMA_ENABLE_DEPRECATION_WARNINGS
 			void EncodeBuffer(FInputOutput* Buffer);
 			void ProcessEncodedBuffer(FInputOutput* Buffer);
 			void Flush();
@@ -126,9 +110,7 @@ namespace AVEncoder
 		};
 
 		FNVENCCommon& NVENC;
-		PRAGMA_DISABLE_DEPRECATION_WARNINGS
 		EVideoFrameFormat FrameFormat = EVideoFrameFormat::Undefined;
-		PRAGMA_ENABLE_DEPRECATION_WARNINGS
 
 		// this could be a TRefCountPtr<ID3D11Device>, CUcontext or void*
 		void* EncoderDevice;
diff --git a/Engine/Plugins/Media/HardwareEncoders/Source/EncoderNVENC/Private/NVENC_EncoderModule.cpp b/Engine/Plugins/Media/HardwareEncoders/Source/EncoderNVENC/Private/NVENC_EncoderModule.cpp
--- a/Engine/Plugins/Media/HardwareEncoders/Source/EncoderNVENC/Private/NVENC_EncoderModule.cpp
+++ b/Engine/Plugins/Media/HardwareEncoders/Source/EncoderNVENC/Private/NVENC_EncoderModule.cpp
@@ -43,11 +43,8 @@ public:
 				{
 					if (IsRHIDeviceNVIDIA())
 					{
-						PRAGMA_DISABLE_DEPRECATION_WARNINGS
 						FVideoEncoderNVENC_H264::Register(FVideoEncoderFactory::Get());
-						PRAGMA_ENABLE_DEPRECATION_WARNINGS
 					}
-					
 				});
 			}
 		}
diff --git a/Engine/Source/Runtime/AVEncoder/Private/AudioEncoder.cpp b/Engine/Source/Runtime/AVEncoder/Private/AudioEncoder.cpp
--- a/Engine/Source/Runtime/AVEncoder/Private/AudioEncoder.cpp
+++ b/Engine/Source/Runtime/AVEncoder/Private/AudioEncoder.cpp
@@ -5,35 +5,26 @@
 
 namespace AVEncoder
 {
-	PRAGMA_DISABLE_DEPRECATION_WARNINGS
 	void FAudioEncoder::RegisterListener(IAudioEncoderListener& Listener)
-	PRAGMA_ENABLE_DEPRECATION_WARNINGS
 	{
 		FScopeLock lock{ &ListenersMutex };
 		check(Listeners.Find(&Listener) == INDEX_NONE);
 		Listeners.AddUnique(&Listener);
 	}
 
-	PRAGMA_DISABLE_DEPRECATION_WARNINGS
 	void FAudioEncoder::UnregisterListener(IAudioEncoderListener& Listener)
-	PRAGMA_ENABLE_DEPRECATION_WARNINGS
 	{
 		FScopeLock lock{ &ListenersMutex };
 		int32 Count = Listeners.Remove(&Listener);
 		check(Count == 1);
 	}
 
-	PRAGMA_DISABLE_DEPRECATION_WARNINGS
 	void FAudioEncoder::OnEncodedAudioFrame(const FMediaPacket& Packet)
-	PRAGMA_ENABLE_DEPRECATION_WARNINGS
 	{
 		FScopeLock lock{ &ListenersMutex };
 		for (auto&& L : Listeners)
 		{
-			PRAGMA_DISABLE_DEPRECATION_WARNINGS
 			L->OnEncodedAudioFrame(Packet);
-			PRAGMA_ENABLE_DEPRECATION_WARNINGS
 		}
-		
 	}
 }
diff --git a/Engine/Source/Runtime/AVEncoder/Private/AudioEncoderFactory.cpp b/Engine/Source/Runtime/AVEncoder/Private/AudioEncoderFactory.cpp
--- a/Engine/Source/Runtime/AVEncoder/Private/AudioEncoderFactory.cpp
+++ b/Engine/Source/Runtime/AVEncoder/Private/AudioEncoderFactory.cpp
@@ -19,29 +19,21 @@ namespace AVEncoder
 		RegisterDefaultFactories();
 	}
 
-	PRAGMA_DISABLE_DEPRECATION_WARNINGS
 	TArray<FAudioEncoderFactory*> FAudioEncoderFactory::Factories;
-	PRAGMA_ENABLE_DEPRECATION_WARNINGS
 
-	PRAGMA_DISABLE_DEPRECATION_WARNINGS
 	void FAudioEncoderFactory::RegisterFactory(FAudioEncoderFactory& Factory)
-	PRAGMA_ENABLE_DEPRECATION_WARNINGS
 	{
 		DoDefaultRegistration();
 
 		Factories.AddUnique(&Factory);
 	}
 
-	PRAGMA_DISABLE_DEPRECATION_WARNINGS
 	void FAudioEncoderFactory::UnregisterFactory(FAudioEncoderFactory& Factory)
-	PRAGMA_ENABLE_DEPRECATION_WARNINGS
 	{
 		Factories.Remove(&Factory);
 	}
 
-	PRAGMA_DISABLE_DEPRECATION_WARNINGS
 	FAudioEncoderFactory* FAudioEncoderFactory::FindFactory(const FString& Codec)
-	PRAGMA_ENABLE_DEPRECATION_WARNINGS
 	{
 		DoDefaultRegistration();
 
@@ -56,9 +48,7 @@ namespace AVEncoder
 		return nullptr;
 	}
 
-	PRAGMA_DISABLE_DEPRECATION_WARNINGS
 	const TArray<FAudioEncoderFactory*> FAudioEncoderFactory::GetAllFactories()
-	PRAGMA_ENABLE_DEPRECATION_WARNINGS
 	{
 		DoDefaultRegistration();
 
@@ -73,12 +63,9 @@ namespace AVEncoder
 #if PLATFORM_WINDOWS
 		// Generic Windows/XBox Wmf encoder
 		static FWmfAudioEncoderFactory WmfAudioEncoderFactory;
-		PRAGMA_DISABLE_DEPRECATION_WARNINGS
 		FAudioEncoderFactory::RegisterFactory(WmfAudioEncoderFactory);
-		PRAGMA_ENABLE_DEPRECATION_WARNINGS
 #endif
 
-
 		// Log all available encoders
 		{
 			auto CodecsInfo = [&](auto&& Factories) -> FString
@@ -86,11 +73,8 @@ namespace AVEncoder
 				FString Str;
 				for (auto&& Factory : Factories)
 				{
-					PRAGMA_DISABLE_DEPRECATION_WARNINGS
 					Str += FString::Printf(TEXT(", %s(%s) "), Factory->GetName(), *FString::Join(Factory->GetSupportedCodecs(), TEXT("/")));
-					PRAGMA_ENABLE_DEPRECATION_WARNINGS
 				}
-				
 				if (Str.IsEmpty())
 				{
 					return FString(TEXT("None"));
@@ -101,10 +85,7 @@ namespace AVEncoder
 				}
 			};
 
-			PRAGMA_DISABLE_DEPRECATION_WARNINGS
 			UE_LOG(LogAVEncoder, Log, TEXT("Available audio encoders: %s "), *CodecsInfo(FAudioEncoderFactory::GetAllFactories()));
-			PRAGMA_ENABLE_DEPRECATION_WARNINGS
 		}
-		
 	}
 }
\ No newline at end of file
diff --git a/Engine/Source/Runtime/AVEncoder/Private/CodecPacket.cpp b/Engine/Source/Runtime/AVEncoder/Private/CodecPacket.cpp
--- a/Engine/Source/Runtime/AVEncoder/Private/CodecPacket.cpp
+++ b/Engine/Source/Runtime/AVEncoder/Private/CodecPacket.cpp
@@ -21,9 +21,7 @@ namespace
 
 namespace AVEncoder
 {
-	PRAGMA_DISABLE_DEPRECATION_WARNINGS
 	FCodecPacket FCodecPacket::Create(const uint8* InData, uint32 InDataSize)
-	PRAGMA_ENABLE_DEPRECATION_WARNINGS
 	{
 		uint8* DataCopy = static_cast<uint8*>(FMemory::Malloc(InDataSize));
 		FMemory::BigBlockMemcpy(DataCopy, InData, InDataSize);
diff --git a/Engine/Source/Runtime/AVEncoder/Private/Decoders/VideoDecoder.cpp b/Engine/Source/Runtime/AVEncoder/Private/Decoders/VideoDecoder.cpp
--- a/Engine/Source/Runtime/AVEncoder/Private/Decoders/VideoDecoder.cpp
+++ b/Engine/Source/Runtime/AVEncoder/Private/Decoders/VideoDecoder.cpp
@@ -11,9 +11,7 @@ namespace AVEncoder
 
 struct FVideoDecoder::FPlatformDecoderAllocInterface
 {
-	PRAGMA_DISABLE_DEPRECATION_WARNINGS
 	FVideoDecoderMethodsWindows	AllocMethods = {};
-	PRAGMA_ENABLE_DEPRECATION_WARNINGS
 	void*						ClientData = nullptr;
 };
 
@@ -29,15 +27,13 @@ bool FVideoDecoder::CreateDecoderAllocationInterface()
 	if (!PlatformDecoderAllocInterface)
 	{
 		check(PlatformDecoderAllocInterface == nullptr);
-
-		PRAGMA_DISABLE_DEPRECATION_WARNINGS
 		FPlatformDecoderAllocInterface* Interface = new FPlatformDecoderAllocInterface {};
-		// TODO: Platform specific!!
+
+// TODO: Platform specific!!
 		Interface->AllocMethods.MagicCookie = 0x57696e58;		// 'WinX'
 		Interface->AllocMethods.This = nullptr;					// This will become the client's 'this' pointer
 		Interface->AllocMethods.GetD3DDevice = nullptr;			// Method to get the current D3D device from the client.
 		Interface->AllocMethods.AllocateFrameBuffer = nullptr;	// Method to allocate frame buffer memory from the client.
-		PRAGMA_ENABLE_DEPRECATION_WARNINGS
 
 		// Call into the client to set up the allocation interface.
 		CreateDecoderAllocationInterfaceFN(&Interface->AllocMethods, &Interface->ClientData);
@@ -57,7 +53,6 @@ bool FVideoDecoder::CreateDecoderAllocationInterface()
 	return true;
 }
 
-PRAGMA_DISABLE_DEPRECATION_WARNINGS
 void FVideoDecoder::ReleaseDecoderAllocationInterface()
 {
 	if (PlatformDecoderAllocInterface)
@@ -67,30 +62,24 @@ void FVideoDecoder::ReleaseDecoderAllocationInterface()
 		delete PlatformDecoderAllocInterface;
 		PlatformDecoderAllocInterface = nullptr;
 	}
+
 }
-PRAGMA_ENABLE_DEPRECATION_WARNINGS
 
-PRAGMA_DISABLE_DEPRECATION_WARNINGS
 void* FVideoDecoder::GetAllocationInterfaceMethods()
 {
 	return PlatformDecoderAllocInterface ? &PlatformDecoderAllocInterface->AllocMethods : nullptr;
-PRAGMA_ENABLE_DEPRECATION_WARNINGS
 }
 
 
-PRAGMA_DISABLE_DEPRECATION_WARNINGS
 EFrameBufferAllocReturn FVideoDecoder::AllocateOutputFrameBuffer(FVideoDecoderAllocFrameBufferResult* OutBuffer, const FVideoDecoderAllocFrameBufferParams* InAllocParams)
-PRAGMA_ENABLE_DEPRECATION_WARNINGS
 {
 	check(InAllocParams);
 	check(OutBuffer);
-	PRAGMA_DISABLE_DEPRECATION_WARNINGS
 	if (PlatformDecoderAllocInterface && PlatformDecoderAllocInterface->AllocMethods.AllocateFrameBuffer.IsBound())
 	{
 		return PlatformDecoderAllocInterface->AllocMethods.AllocateFrameBuffer.Execute(PlatformDecoderAllocInterface->AllocMethods.This, InAllocParams, OutBuffer);
 	}
 	return EFrameBufferAllocReturn::CODEC_Failure;
-	PRAGMA_ENABLE_DEPRECATION_WARNINGS
 }
 
 
diff --git a/Engine/Source/Runtime/AVEncoder/Private/Decoders/VideoDecoderH264_Dummy.cpp b/Engine/Source/Runtime/AVEncoder/Private/Decoders/VideoDecoderH264_Dummy.cpp
--- a/Engine/Source/Runtime/AVEncoder/Private/Decoders/VideoDecoderH264_Dummy.cpp
+++ b/Engine/Source/Runtime/AVEncoder/Private/Decoders/VideoDecoderH264_Dummy.cpp
@@ -8,22 +8,16 @@
 namespace AVEncoder
 {
 
-PRAGMA_DISABLE_DEPRECATION_WARNINGS
 class FVideoDecoderOutputDummy : public FVideoDecoderOutput
-PRAGMA_ENABLE_DEPRECATION_WARNINGS
 {
 public:
-	PRAGMA_DISABLE_DEPRECATION_WARNINGS
 	FVideoDecoderOutputDummy(int32 w, int32 h, int64 pts)
 	: Width(w), Height(h), PTS(pts)
 	{}
-	PRAGMA_ENABLE_DEPRECATION_WARNINGS
 
-	PRAGMA_DISABLE_DEPRECATION_WARNINGS
 	virtual ~FVideoDecoderOutputDummy()
 	{
 	}
-	PRAGMA_ENABLE_DEPRECATION_WARNINGS
 
 	virtual int32 AddRef() override
 	{
@@ -56,9 +50,7 @@ public:
 		return PTS;
 	}
 
-	PRAGMA_DISABLE_DEPRECATION_WARNINGS
 	virtual const FVideoDecoderAllocFrameBufferResult* GetAllocatedBuffer() const
-	PRAGMA_ENABLE_DEPRECATION_WARNINGS
 	{
 		return &Buffer;
 	}
@@ -101,17 +93,13 @@ public:
 	}
 
 	// Internal for allocation.
-	PRAGMA_DISABLE_DEPRECATION_WARNINGS
 	FVideoDecoderAllocFrameBufferResult* GetBuffer()
-	PRAGMA_ENABLE_DEPRECATION_WARNINGS
 	{
 		return &Buffer;
 	}
 
 private:
-	PRAGMA_DISABLE_DEPRECATION_WARNINGS
 	FVideoDecoderAllocFrameBufferResult	Buffer = {};
-	PRAGMA_ENABLE_DEPRECATION_WARNINGS
 	int32	RefCount = 1;
 	int32	Width = 0;
 	int32	Height = 0;
@@ -121,44 +109,30 @@ private:
 	// query whether or not encoder is supported and available
 //	static bool GetIsAvailable(FVideoEncoderInputImpl& InInput, FVideoEncoderInfo& OutEncoderInfo);
 
-PRAGMA_DISABLE_DEPRECATION_WARNINGS
 void FVideoDecoderH264_Dummy::Register(FVideoDecoderFactory& InFactory)
-PRAGMA_ENABLE_DEPRECATION_WARNINGS
 {
-	PRAGMA_DISABLE_DEPRECATION_WARNINGS
 	FVideoDecoderInfo	DecoderInfo;
 	DecoderInfo.CodecType = ECodecType::H264;
 	DecoderInfo.MaxWidth = 1920;
 	DecoderInfo.MaxHeight = 1088;
-	PRAGMA_ENABLE_DEPRECATION_WARNINGS
 
-	PRAGMA_DISABLE_DEPRECATION_WARNINGS
 	InFactory.Register(DecoderInfo, []() {
-		return new FVideoDecoderH264_Dummy();
-	});
-	PRAGMA_ENABLE_DEPRECATION_WARNINGS
+			return new FVideoDecoderH264_Dummy();
+		});
 }
 
-PRAGMA_DISABLE_DEPRECATION_WARNINGS
 FVideoDecoderH264_Dummy::FVideoDecoderH264_Dummy()
 {
 }
-PRAGMA_ENABLE_DEPRECATION_WARNINGS
 
-PRAGMA_DISABLE_DEPRECATION_WARNINGS
 FVideoDecoderH264_Dummy::~FVideoDecoderH264_Dummy()
 {
 }
-PRAGMA_ENABLE_DEPRECATION_WARNINGS
 
-PRAGMA_DISABLE_DEPRECATION_WARNINGS
 bool FVideoDecoderH264_Dummy::Setup(const FVideoDecoder::FInit& InInit)
-PRAGMA_ENABLE_DEPRECATION_WARNINGS
 {
-	PRAGMA_DISABLE_DEPRECATION_WARNINGS
 	CreateDecoderAllocationInterfaceFN = InInit.CreateDecoderAllocationInterface;
 	ReleaseDecoderAllocationInterfaceFN = InInit.ReleaseDecoderAllocationInterface;
-	PRAGMA_ENABLE_DEPRECATION_WARNINGS
 	return true;
 }
 
@@ -166,41 +140,30 @@ void FVideoDecoderH264_Dummy::Shutdown()
 {
 	if (bIsInitialized)
 	{
-		PRAGMA_DISABLE_DEPRECATION_WARNINGS
 		ReleaseDecoderAllocationInterface();
-		PRAGMA_ENABLE_DEPRECATION_WARNINGS
 		bIsInitialized = false;
 	}
 	delete this;
 }
 
 
-PRAGMA_DISABLE_DEPRECATION_WARNINGS
 FVideoDecoder::EDecodeResult FVideoDecoderH264_Dummy::Decode(const FVideoDecoderInput* InInput)
-PRAGMA_ENABLE_DEPRECATION_WARNINGS
 {
 	// Lazily initialize the decoder.
 	if (!bIsInitialized)
 	{
-		PRAGMA_DISABLE_DEPRECATION_WARNINGS
 		bIsInitialized = CreateDecoderAllocationInterface();
-		PRAGMA_ENABLE_DEPRECATION_WARNINGS
 	}
-	
 
 	//
-	PRAGMA_DISABLE_DEPRECATION_WARNINGS
 	if (OnDecodedFrame)
-	PRAGMA_ENABLE_DEPRECATION_WARNINGS
 	{
 		if (bIsInitialized)
 		{
-			PRAGMA_DISABLE_DEPRECATION_WARNINGS
 			int32_t	Width = InInput->GetWidth();
 			int32_t Height = InInput->GetHeight();
-			
-			TUniquePtr<FVideoDecoderOutputDummy> pNew(new FVideoDecoderOutputDummy(Width, Height, InInput->GetPTS()));
 
+			TUniquePtr<FVideoDecoderOutputDummy> pNew(new FVideoDecoderOutputDummy(Width, Height, InInput->GetPTS()));
 
 			// Get memory from the application
 			FVideoDecoderAllocFrameBufferParams ap {};
@@ -229,13 +192,10 @@ PRAGMA_ENABLE_DEPRECATION_WARNINGS
 				// Error!
 				return FVideoDecoder::EDecodeResult::Failure;
 			}
-			PRAGMA_ENABLE_DEPRECATION_WARNINGS
 
 			// Buffer needs to contain an NV12 texture.
 			int32 ySize = Width * Height;
-			PRAGMA_DISABLE_DEPRECATION_WARNINGS
 			uint8* pY = (uint8*)pNew->GetBuffer()->AllocatedBuffer;
-			PRAGMA_ENABLE_DEPRECATION_WARNINGS
 			uint8* pUV = pY + ySize;
 			FMemory::Memset(pUV, 128, ySize / 2);
 			int C = ++YOffset;
@@ -244,25 +204,19 @@ PRAGMA_ENABLE_DEPRECATION_WARNINGS
 				FMemory::Memset(pY, (uint8)C++, Width);
 				pY += Width;
 			}
-			PRAGMA_DISABLE_DEPRECATION_WARNINGS
 			OnDecodedFrame(pNew.Release());
 			return FVideoDecoder::EDecodeResult::Success;
-			PRAGMA_ENABLE_DEPRECATION_WARNINGS
 		}
 		else
 		{
 // TODO: could not decode, but may need to still notify the caller with an empty frame??
-			PRAGMA_DISABLE_DEPRECATION_WARNINGS
 			return FVideoDecoder::EDecodeResult::Failure;
-			PRAGMA_ENABLE_DEPRECATION_WARNINGS
 		}
 	}
 	else
 	{
 		// With no registered callback that's interested in the result we can presume we would have been successful.
-		PRAGMA_DISABLE_DEPRECATION_WARNINGS
 		return FVideoDecoder::EDecodeResult::Success;
-		PRAGMA_ENABLE_DEPRECATION_WARNINGS
 	}
 }
 
diff --git a/Engine/Source/Runtime/AVEncoder/Private/Decoders/VideoDecoderH264_Dummy.h b/Engine/Source/Runtime/AVEncoder/Private/Decoders/VideoDecoderH264_Dummy.h
--- a/Engine/Source/Runtime/AVEncoder/Private/Decoders/VideoDecoderH264_Dummy.h
+++ b/Engine/Source/Runtime/AVEncoder/Private/Decoders/VideoDecoderH264_Dummy.h
@@ -12,20 +12,16 @@
 namespace AVEncoder
 {
 
-PRAGMA_DISABLE_DEPRECATION_WARNINGS
 class FVideoDecoderH264_Dummy : public FVideoDecoder
-PRAGMA_ENABLE_DEPRECATION_WARNINGS
 {
 public:
 	// register decoder with video decoder factory
-	PRAGMA_DISABLE_DEPRECATION_WARNINGS
 	static void Register(FVideoDecoderFactory& InFactory);
-	virtual bool Setup(const FInit& InInit) override;
-	virtual EDecodeResult Decode(const FVideoDecoderInput* InInput) override;
-	PRAGMA_ENABLE_DEPRECATION_WARNINGS
 
+	virtual bool Setup(const FInit& InInit) override;
 	virtual void Shutdown() override;
 
+	virtual EDecodeResult Decode(const FVideoDecoderInput* InInput) override;
 
 
 	// query whether or not encoder is supported and available
diff --git a/Engine/Source/Runtime/AVEncoder/Private/Decoders/VideoDecoderH264_Impl.cpp b/Engine/Source/Runtime/AVEncoder/Private/Decoders/VideoDecoderH264_Impl.cpp
--- a/Engine/Source/Runtime/AVEncoder/Private/Decoders/VideoDecoderH264_Impl.cpp
+++ b/Engine/Source/Runtime/AVEncoder/Private/Decoders/VideoDecoderH264_Impl.cpp
@@ -21,9 +21,7 @@ namespace AVEncoder
 #if CODEC_HAVE_NATIVE_H264_DECODER
 
 #if PLATFORM_WINDOWS
-PRAGMA_DISABLE_DEPRECATION_WARNINGS
 void FVideoDecoderH264_Impl::Register(FVideoDecoderFactory& InFactory)
-PRAGMA_ENABLE_DEPRECATION_WARNINGS
 {
 	FVideoDecoderH264_Windows::Register(InFactory);
 }
@@ -31,9 +29,7 @@ PRAGMA_ENABLE_DEPRECATION_WARNINGS
 
 #else
 
-PRAGMA_DISABLE_DEPRECATION_WARNINGS
 void FVideoDecoderH264_Impl::Register(FVideoDecoderFactory& InFactory)
-PRAGMA_ENABLE_DEPRECATION_WARNINGS
 {
 }
 
diff --git a/Engine/Source/Runtime/AVEncoder/Private/Decoders/VideoDecoderH264_Impl.h b/Engine/Source/Runtime/AVEncoder/Private/Decoders/VideoDecoderH264_Impl.h
--- a/Engine/Source/Runtime/AVEncoder/Private/Decoders/VideoDecoderH264_Impl.h
+++ b/Engine/Source/Runtime/AVEncoder/Private/Decoders/VideoDecoderH264_Impl.h
@@ -8,25 +8,18 @@
 namespace AVEncoder
 {
 
-PRAGMA_DISABLE_DEPRECATION_WARNINGS
 class FVideoDecoderH264_Impl : public FVideoDecoder
-PRAGMA_ENABLE_DEPRECATION_WARNINGS
 {
 public:
 	// register decoder with video decoder factory
-	PRAGMA_DISABLE_DEPRECATION_WARNINGS
 	static void Register(FVideoDecoderFactory& InFactory);
-	virtual bool Setup(const FInit& InInit) = 0;
-	virtual EDecodeResult Decode(const FVideoDecoderInput* InInput) = 0;
-	PRAGMA_ENABLE_DEPRECATION_WARNINGS
 
+	virtual bool Setup(const FInit& InInit) = 0;
 	virtual void Shutdown() = 0;
 
+	virtual EDecodeResult Decode(const FVideoDecoderInput* InInput) = 0;
 protected:
-	PRAGMA_DISABLE_DEPRECATION_WARNINGS
 	virtual ~FVideoDecoderH264_Impl() = default;
-	PRAGMA_ENABLE_DEPRECATION_WARNINGS
 };
 
-
 } // namespace AVEncoder
diff --git a/Engine/Source/Runtime/AVEncoder/Private/Decoders/VideoDecoderMpeg4.cpp b/Engine/Source/Runtime/AVEncoder/Private/Decoders/VideoDecoderMpeg4.cpp
--- a/Engine/Source/Runtime/AVEncoder/Private/Decoders/VideoDecoderMpeg4.cpp
+++ b/Engine/Source/Runtime/AVEncoder/Private/Decoders/VideoDecoderMpeg4.cpp
@@ -17,22 +17,16 @@
 namespace AVEncoder
 {
 
-PRAGMA_DISABLE_DEPRECATION_WARNINGS
 class FVideoDecoderOutputMPEG4 : public FVideoDecoderOutput
-PRAGMA_ENABLE_DEPRECATION_WARNINGS
 {
 public:
-	PRAGMA_DISABLE_DEPRECATION_WARNINGS
 	FVideoDecoderOutputMPEG4(int32 w, int32 h, int64 pts)
 	: Width(w), Height(h), PTS(pts)
 	{}
-	PRAGMA_ENABLE_DEPRECATION_WARNINGS
 
-	PRAGMA_DISABLE_DEPRECATION_WARNINGS
 	virtual ~FVideoDecoderOutputMPEG4()
 	{
 	}
-	PRAGMA_ENABLE_DEPRECATION_WARNINGS
 
 	virtual int32 AddRef() override
 	{
@@ -65,9 +59,7 @@ public:
 		return PTS;
 	}
 
-	PRAGMA_DISABLE_DEPRECATION_WARNINGS
 	virtual const FVideoDecoderAllocFrameBufferResult* GetAllocatedBuffer() const
-	PRAGMA_ENABLE_DEPRECATION_WARNINGS
 	{
 		return &Buffer;
 	}
@@ -114,17 +106,13 @@ public:
 	}
 
 	// Internal for allocation.
-	PRAGMA_DISABLE_DEPRECATION_WARNINGS
 	FVideoDecoderAllocFrameBufferResult* GetBuffer()
-	PRAGMA_ENABLE_DEPRECATION_WARNINGS
 	{
 		return &Buffer;
 	}
 
 private:
-	PRAGMA_DISABLE_DEPRECATION_WARNINGS
 	FVideoDecoderAllocFrameBufferResult	Buffer = {};
-	PRAGMA_ENABLE_DEPRECATION_WARNINGS
 	int32	RefCount = 1;
 	int32	Width = 0;
 	int32	Pitch = 0;
@@ -147,14 +135,10 @@ static TMap<void*, uint32>& Actives()
 class FVideoDecoderMPEG4_Impl : public FVideoDecoderMPEG4, public vdecmpeg4::VIDStreamIO, public vdecmpeg4::VIDStreamEvents
 {
 public:
-	PRAGMA_DISABLE_DEPRECATION_WARNINGS
 	virtual bool Setup(const FInit& InInit) override;
-	PRAGMA_ENABLE_DEPRECATION_WARNINGS
 	virtual void Shutdown() override;
 
-	PRAGMA_DISABLE_DEPRECATION_WARNINGS
 	virtual EDecodeResult Decode(const FVideoDecoderInput* InInput) override;
-	PRAGMA_ENABLE_DEPRECATION_WARNINGS
 
 	FVideoDecoderMPEG4_Impl();
 	virtual ~FVideoDecoderMPEG4_Impl();
@@ -212,22 +196,16 @@ protected:
 
 
 
-PRAGMA_DISABLE_DEPRECATION_WARNINGS
 void FVideoDecoderMPEG4::Register(FVideoDecoderFactory& InFactory)
-PRAGMA_ENABLE_DEPRECATION_WARNINGS
 {
-	PRAGMA_DISABLE_DEPRECATION_WARNINGS
 	FVideoDecoderInfo	DecoderInfo;
 	DecoderInfo.CodecType = ECodecType::MPEG4;
 	DecoderInfo.MaxWidth = 1920;
 	DecoderInfo.MaxHeight = 1088;
-	PRAGMA_ENABLE_DEPRECATION_WARNINGS
 
-	PRAGMA_DISABLE_DEPRECATION_WARNINGS
 	InFactory.Register(DecoderInfo, []() {
-		return new FVideoDecoderMPEG4_Impl();
-	});
-	PRAGMA_ENABLE_DEPRECATION_WARNINGS
+			return new FVideoDecoderMPEG4_Impl();
+		});
 }
 
 int32 gTestStreamIndex4 = 0;
@@ -246,14 +224,10 @@ FVideoDecoderMPEG4_Impl::~FVideoDecoderMPEG4_Impl()
 	check(!DecoderHandle);
 }
 
-PRAGMA_DISABLE_DEPRECATION_WARNINGS
 bool FVideoDecoderMPEG4_Impl::Setup(const FVideoDecoder::FInit& InInit)
-PRAGMA_ENABLE_DEPRECATION_WARNINGS
 {
-	PRAGMA_DISABLE_DEPRECATION_WARNINGS
 	CreateDecoderAllocationInterfaceFN = InInit.CreateDecoderAllocationInterface;
 	ReleaseDecoderAllocationInterfaceFN = InInit.ReleaseDecoderAllocationInterface;
-	PRAGMA_ENABLE_DEPRECATION_WARNINGS
 	return true;
 }
 
@@ -266,9 +240,7 @@ void FVideoDecoderMPEG4_Impl::Shutdown()
 			vdecmpeg4::VIDDestroyDecoder(DecoderHandle);
 			DecoderHandle = nullptr;
 		}
-		PRAGMA_DISABLE_DEPRECATION_WARNINGS
 		ReleaseDecoderAllocationInterface();
-		PRAGMA_ENABLE_DEPRECATION_WARNINGS
 		bIsInitialized = false;
 	}
 	delete this;
@@ -278,9 +250,7 @@ bool FVideoDecoderMPEG4_Impl::FirstUseInit()
 {
 	if (!bIsInitialized)
 	{
-		PRAGMA_DISABLE_DEPRECATION_WARNINGS
 		bIsInitialized = CreateDecoderAllocationInterface();
-		PRAGMA_ENABLE_DEPRECATION_WARNINGS
 		if (bIsInitialized)
 		{
 			FMemory::Memzero(DecoderSetup);
@@ -355,9 +325,7 @@ bool FVideoDecoderMPEG4_Impl::IsEof()
 }
 
 
-PRAGMA_DISABLE_DEPRECATION_WARNINGS
 static void CopyI420ToNV12(const FVideoDecoderAllocFrameBufferResult* OutBuf, const vdecmpeg4::VIDImage* vid)
-PRAGMA_ENABLE_DEPRECATION_WARNINGS
 {
 	int32 Width = vid->width;
 	int32 Height = vid->height;
@@ -367,7 +335,6 @@ PRAGMA_ENABLE_DEPRECATION_WARNINGS
 	const uint8_t* srcV = vid->v;
 
 	// Allocated buffer needs to have 3 planes.
-	PRAGMA_DISABLE_DEPRECATION_WARNINGS
 	check(OutBuf->AllocatedPlanesNum == 3);
 
 	uint8* OutBufferBase = (uint8*)OutBuf->AllocatedBuffer;
@@ -375,28 +342,22 @@ PRAGMA_ENABLE_DEPRECATION_WARNINGS
 	check(OutBuf->AllocatedPlaneDesc[0].BytesPerPixel == 1);
 	check(OutBuf->AllocatedPlaneDesc[0].ByteOffsetBetweenPixels == 1);
 	uint8* dstY = OutBufferBase + OutBuf->AllocatedPlaneDesc[0].ByteOffsetToFirstPixel;
-	PRAGMA_ENABLE_DEPRECATION_WARNINGS
 	// Need to copy the Y plane row by row.
 	for(int32 y=0; y<Height; ++y)
 	{
 		FMemory::Memcpy(dstY, srcY, Width);
 		srcY += vid->texWidth;
-		PRAGMA_DISABLE_DEPRECATION_WARNINGS
 		dstY += OutBuf->AllocatedPlaneDesc[0].ByteOffsetBetweenRows;
-		PRAGMA_ENABLE_DEPRECATION_WARNINGS
 	}
-	
 
 	// The U and V plane must be interleaved for NV12. We don't specifically do interleaving here but
 	// instead rely on the output plane description to be set up accordingly.
-	PRAGMA_DISABLE_DEPRECATION_WARNINGS
 	check(OutBuf->AllocatedPlaneDesc[1].BytesPerPixel == 1);
 	check(OutBuf->AllocatedPlaneDesc[2].BytesPerPixel == 1);
 	uint8* dstU = OutBufferBase + OutBuf->AllocatedPlaneDesc[1].ByteOffsetToFirstPixel;
 	uint8* dstV = OutBufferBase + OutBuf->AllocatedPlaneDesc[2].ByteOffsetToFirstPixel;
 	const int32 uOffCol = OutBuf->AllocatedPlaneDesc[1].ByteOffsetBetweenPixels;
 	const int32 vOffCol = OutBuf->AllocatedPlaneDesc[2].ByteOffsetBetweenPixels;
-	PRAGMA_ENABLE_DEPRECATION_WARNINGS
 	for(int32 v=0, vMax=Height/2; v<vMax; ++v)
 	{
 		uint8* U = dstU;
@@ -410,36 +371,27 @@ PRAGMA_ENABLE_DEPRECATION_WARNINGS
 		}
 		srcU += (vid->texWidth - vid->width) / 2;
 		srcV += (vid->texWidth - vid->width) / 2;
-		PRAGMA_DISABLE_DEPRECATION_WARNINGS
 		dstU += OutBuf->AllocatedPlaneDesc[1].ByteOffsetBetweenRows;
 		dstV += OutBuf->AllocatedPlaneDesc[2].ByteOffsetBetweenRows;
-		PRAGMA_ENABLE_DEPRECATION_WARNINGS
 	}
 }
 
 
-PRAGMA_DISABLE_DEPRECATION_WARNINGS
 FVideoDecoder::EDecodeResult FVideoDecoderMPEG4_Impl::Decode(const FVideoDecoderInput* InInput)
-PRAGMA_ENABLE_DEPRECATION_WARNINGS
 {
 	// With no registered callback that's interested in the result we can presume we would have been successful.
-	PRAGMA_DISABLE_DEPRECATION_WARNINGS
 	if (!OnDecodedFrame)
 	{
 		return FVideoDecoder::EDecodeResult::Success;
 	}
-	PRAGMA_ENABLE_DEPRECATION_WARNINGS
 	// Initialize decoder on first decode call.
 	if (!FirstUseInit())
 	{
-		PRAGMA_DISABLE_DEPRECATION_WARNINGS
 		return FVideoDecoder::EDecodeResult::Failure;
-	}PRAGMA_ENABLE_DEPRECATION_WARNINGS
-	
+	}
 
 	// Setup an access unit to run through the decoder.
 	TUniquePtr<FInDecoderData> AU = MakeUnique<FInDecoderData>();
-	PRAGMA_DISABLE_DEPRECATION_WARNINGS
 	if (InInput->GetDataSize() <= 0)
 	{
 		return FVideoDecoder::EDecodeResult::Failure;
@@ -452,7 +404,6 @@ PRAGMA_ENABLE_DEPRECATION_WARNINGS
 	AU->Height = InInput->GetHeight();
 	AU->bIsKeyframe = InInput->IsKeyframe();
 	AU->bIsComplete = InInput->IsCompleteFrame();
-	PRAGMA_ENABLE_DEPRECATION_WARNINGS
 	PendingDecodeData.Enqueue(MoveTemp(AU));
 
 	// Decode all pending input.
@@ -475,9 +426,7 @@ PRAGMA_ENABLE_DEPRECATION_WARNINGS
 			else
 			{
 				UE_LOG(LogVideoDecoder, Error, TEXT("VIDCreateDecoder() failed with %d"), (int)LastDecoderError);
-				PRAGMA_DISABLE_DEPRECATION_WARNINGS
 				return FVideoDecoder::EDecodeResult::Failure;
-				PRAGMA_ENABLE_DEPRECATION_WARNINGS
 				//vdecmpeg4::VIDDestroyDecoder(DecoderHandle);
 				//DecoderHandle = nullptr;
 			}
@@ -491,7 +440,6 @@ PRAGMA_ENABLE_DEPRECATION_WARNINGS
 				int32_t	Width = frame->width;
 				int32_t Height = frame->height;
 
-				PRAGMA_DISABLE_DEPRECATION_WARNINGS
 				TUniquePtr<FVideoDecoderOutputMPEG4> pNew(new FVideoDecoderOutputMPEG4(Width, Height, InInput->GetPTS()));
 
 				// Get memory from the application
@@ -534,16 +482,13 @@ PRAGMA_ENABLE_DEPRECATION_WARNINGS
 				{
 					pNew->SetPitchX(ap.Width);
 				}
-				PRAGMA_ENABLE_DEPRECATION_WARNINGS
 
 				// Copy the image across, converting it into NV12 format.
 				CopyI420ToNV12(pNew->GetBuffer(), frame);
 				frame->Release();
 				// Deliver
-				PRAGMA_DISABLE_DEPRECATION_WARNINGS
 				OnDecodedFrame(pNew.Release());
-				PRAGMA_ENABLE_DEPRECATION_WARNINGS
-			}	
+			}
 		}
 		else if (LastDecoderError == vdecmpeg4::VID_ERROR_STREAM_UNDERFLOW)
 		{
@@ -558,9 +503,7 @@ PRAGMA_ENABLE_DEPRECATION_WARNINGS
 			// Error!
 			UE_LOG(LogVideoDecoder, Error, TEXT("VIDStreamDecode() failed with %d"), (int)LastDecoderError);
 // TODO: destroy decoder and wait for next keyframe?
-			PRAGMA_DISABLE_DEPRECATION_WARNINGS
 			return FVideoDecoder::EDecodeResult::Failure;
-			PRAGMA_ENABLE_DEPRECATION_WARNINGS
 			break;
 		}
 
@@ -570,11 +513,8 @@ PRAGMA_ENABLE_DEPRECATION_WARNINGS
 			CurrentAU.Reset();
 		}
 	}
-	PRAGMA_DISABLE_DEPRECATION_WARNINGS
 	return FVideoDecoder::EDecodeResult::Success;
-	PRAGMA_ENABLE_DEPRECATION_WARNINGS
 }
 
-
 } // namespace AVEncoder
 
diff --git a/Engine/Source/Runtime/AVEncoder/Private/Decoders/VideoDecoderMpeg4.h b/Engine/Source/Runtime/AVEncoder/Private/Decoders/VideoDecoderMpeg4.h
--- a/Engine/Source/Runtime/AVEncoder/Private/Decoders/VideoDecoderMpeg4.h
+++ b/Engine/Source/Runtime/AVEncoder/Private/Decoders/VideoDecoderMpeg4.h
@@ -8,26 +8,20 @@
 namespace AVEncoder
 {
 
-PRAGMA_DISABLE_DEPRECATION_WARNINGS
 class FVideoDecoderMPEG4 : public FVideoDecoder
-PRAGMA_ENABLE_DEPRECATION_WARNINGS
 {
 public:
 	// register decoder with video decoder factory
-	PRAGMA_DISABLE_DEPRECATION_WARNINGS
 	static void Register(FVideoDecoderFactory& InFactory);
-	virtual bool Setup(const FInit& InInit) = 0;
-	virtual EDecodeResult Decode(const FVideoDecoderInput* InInput) = 0;
-	PRAGMA_ENABLE_DEPRECATION_WARNINGS
 
+	virtual bool Setup(const FInit& InInit) = 0;
 	virtual void Shutdown() = 0;
 
+	virtual EDecodeResult Decode(const FVideoDecoderInput* InInput) = 0;
 
 protected:
-	PRAGMA_DISABLE_DEPRECATION_WARNINGS
 	virtual ~FVideoDecoderMPEG4() = default;
 	FVideoDecoderMPEG4() = default;
-	PRAGMA_ENABLE_DEPRECATION_WARNINGS
 };
 
 } // namespace AVEncoder
diff --git a/Engine/Source/Runtime/AVEncoder/Private/Decoders/VideoDecoderVP8.cpp b/Engine/Source/Runtime/AVEncoder/Private/Decoders/VideoDecoderVP8.cpp
--- a/Engine/Source/Runtime/AVEncoder/Private/Decoders/VideoDecoderVP8.cpp
+++ b/Engine/Source/Runtime/AVEncoder/Private/Decoders/VideoDecoderVP8.cpp
@@ -16,9 +16,7 @@
 namespace AVEncoder
 {
 
-PRAGMA_DISABLE_DEPRECATION_WARNINGS
 void FVideoDecoderVP8::Register(FVideoDecoderFactory& InFactory)
-PRAGMA_ENABLE_DEPRECATION_WARNINGS
 {
 }
 
diff --git a/Engine/Source/Runtime/AVEncoder/Private/Decoders/VideoDecoderVP8.h b/Engine/Source/Runtime/AVEncoder/Private/Decoders/VideoDecoderVP8.h
--- a/Engine/Source/Runtime/AVEncoder/Private/Decoders/VideoDecoderVP8.h
+++ b/Engine/Source/Runtime/AVEncoder/Private/Decoders/VideoDecoderVP8.h
@@ -8,20 +8,16 @@
 namespace AVEncoder
 {
 
-PRAGMA_DISABLE_DEPRECATION_WARNINGS
 class FVideoDecoderVP8 : public FVideoDecoder
-PRAGMA_ENABLE_DEPRECATION_WARNINGS
 {
 public:
 	// register decoder with video decoder factory
-	PRAGMA_DISABLE_DEPRECATION_WARNINGS
 	static void Register(FVideoDecoderFactory& InFactory);
-	virtual bool Setup(const FInit& InInit) = 0;
-	virtual EDecodeResult Decode(const FVideoDecoderInput* InInput) = 0;
-	PRAGMA_ENABLE_DEPRECATION_WARNINGS
 
+	virtual bool Setup(const FInit& InInit) = 0;
 	virtual void Shutdown() = 0;
 
+	virtual EDecodeResult Decode(const FVideoDecoderInput* InInput) = 0;
 protected:
 	virtual ~FVideoDecoderVP8() = default;
 };
diff --git a/Engine/Source/Runtime/AVEncoder/Private/Decoders/Windows/VideoDecoderH264_Windows.cpp b/Engine/Source/Runtime/AVEncoder/Private/Decoders/Windows/VideoDecoderH264_Windows.cpp
--- a/Engine/Source/Runtime/AVEncoder/Private/Decoders/Windows/VideoDecoderH264_Windows.cpp
+++ b/Engine/Source/Runtime/AVEncoder/Private/Decoders/Windows/VideoDecoderH264_Windows.cpp
@@ -417,20 +417,14 @@ void FF5PlayerVideoDecoderOutputDX::ShutdownPoolable()
 /**
  * Decoded media sample
  */
-PRAGMA_DISABLE_DEPRECATION_WARNINGS
 class FVideoDecoderOutputH264_Windows : public FVideoDecoderOutput
-PRAGMA_ENABLE_DEPRECATION_WARNINGS
 {
 public:
-	PRAGMA_DISABLE_DEPRECATION_WARNINGS
 	FVideoDecoderOutputH264_Windows()
-	PRAGMA_ENABLE_DEPRECATION_WARNINGS
 	{}
 
-	PRAGMA_DISABLE_DEPRECATION_WARNINGS
 	virtual ~FVideoDecoderOutputH264_Windows()
 	{}
-	PRAGMA_ENABLE_DEPRECATION_WARNINGS
 
 	virtual int32 AddRef() override
 	{
@@ -486,9 +480,7 @@ public:
 		PTS = InPTS;
 	}
 
-	PRAGMA_DISABLE_DEPRECATION_WARNINGS
 	virtual const FVideoDecoderAllocFrameBufferResult* GetAllocatedBuffer() const
-	PRAGMA_ENABLE_DEPRECATION_WARNINGS
 	{
 		return &Buffer;
 	}
@@ -567,9 +559,7 @@ public:
 	}
 
 	// Internal for allocation.
-	PRAGMA_DISABLE_DEPRECATION_WARNINGS
 	FVideoDecoderAllocFrameBufferResult* GetBuffer()
-	PRAGMA_ENABLE_DEPRECATION_WARNINGS
 	{
 		return &Buffer;
 	}
@@ -581,9 +571,7 @@ public:
 	}
 
 private:
-	PRAGMA_DISABLE_DEPRECATION_WARNINGS
 	FVideoDecoderAllocFrameBufferResult	Buffer = {};
-	PRAGMA_ENABLE_DEPRECATION_WARNINGS
 	int32	RefCount = 1;
 	int32	Width = 0;
 	int32	Height = 0;
@@ -613,14 +601,10 @@ public:
 
 	virtual ~FVideoDecoderH264_WindowsImpl();
 
-	PRAGMA_DISABLE_DEPRECATION_WARNINGS
 	virtual bool Setup(const FInit& InInit) override;
-	PRAGMA_ENABLE_DEPRECATION_WARNINGS
 	virtual void Shutdown() override;
 
-	PRAGMA_DISABLE_DEPRECATION_WARNINGS
 	virtual EDecodeResult Decode(const FVideoDecoderInput* InInput) override;
-	PRAGMA_ENABLE_DEPRECATION_WARNINGS
 
 
 
@@ -821,20 +805,16 @@ bool						FVideoDecoderH264_WindowsImpl::bGlobalMFStartedUp = false;
 /**
  * Registers this decoder with the decoder factory.
  */
-PRAGMA_DISABLE_DEPRECATION_WARNINGS
 void FVideoDecoderH264_Windows::Register(FVideoDecoderFactory& InFactory)
-PRAGMA_ENABLE_DEPRECATION_WARNINGS
 {
-	PRAGMA_DISABLE_DEPRECATION_WARNINGS
 	FVideoDecoderInfo	DecoderInfo;
 	DecoderInfo.CodecType = ECodecType::H264;
 	DecoderInfo.MaxWidth = 1920;
 	DecoderInfo.MaxHeight = 1088;
-	
+
 	InFactory.Register(DecoderInfo, []() {
 			return new FVideoDecoderH264_WindowsImpl();
 		});
-	PRAGMA_ENABLE_DEPRECATION_WARNINGS
 }
 
 /*********************************************************************************************************************/
@@ -907,18 +887,12 @@ bool FVideoDecoderH264_WindowsImpl::CreateD3DResources()
 	void* ApplicationD3DDevice = nullptr;
 	int32_t ApplicationD3DDeviceVersion = 0;
 	bool bOk = false;
-	PRAGMA_DISABLE_DEPRECATION_WARNINGS
 	FVideoDecoderMethodsWindows* AllocationMethods = reinterpret_cast<FVideoDecoderMethodsWindows*>(GetAllocationInterfaceMethods());
-	PRAGMA_ENABLE_DEPRECATION_WARNINGS
 	if (AllocationMethods)
 	{
-		PRAGMA_DISABLE_DEPRECATION_WARNINGS
 		if (AllocationMethods->MagicCookie == 0x57696e58) // 'WinX'
-		PRAGMA_ENABLE_DEPRECATION_WARNINGS
 		{
-			PRAGMA_DISABLE_DEPRECATION_WARNINGS
 			int32_t AppResult = AllocationMethods->GetD3DDevice.Execute(AllocationMethods->This, &ApplicationD3DDevice, &ApplicationD3DDeviceVersion);
-			PRAGMA_ENABLE_DEPRECATION_WARNINGS
 			if (AppResult == 0)
 			{
 				// D3D version is returned as major*1000+minor.
@@ -945,9 +919,7 @@ bool FVideoDecoderH264_WindowsImpl::CreateD3DResources()
 		}
 		else
 		{
-			PRAGMA_DISABLE_DEPRECATION_WARNINGS
 			UE_LOG(LogVideoDecoder, Error, TEXT("Incorrect video decoder allocation interface is set. Is %08x, should be %08x"), AllocationMethods->MagicCookie, 0x57696e58);
-			PRAGMA_ENABLE_DEPRECATION_WARNINGS
 		}
 	}
 	else
@@ -1629,14 +1601,11 @@ bool FVideoDecoderH264_WindowsImpl::AllocateApplicationOutputSampleBuffer(int32
 
 	// Get the buffer for the output sample from the application. This is the part that gets wrapped in a
 	// webrtc structure and passed back out to the application.
-	PRAGMA_DISABLE_DEPRECATION_WARNINGS
 	FVideoDecoderAllocFrameBufferParams ap {};
 	EFrameBufferAllocReturn ar;
-	PRAGMA_ENABLE_DEPRECATION_WARNINGS
 	if (DecodeMode == EDecodeMode::Win8HW)
 	{
 		// We will pass a ID3D11Texture2D pointer to a texture of _our_ decoding device to _share_ with the renderer.
-		PRAGMA_DISABLE_DEPRECATION_WARNINGS
 		ap.FrameBufferType = EFrameBufferType::CODEC_TextureHandle;
 		ap.AllocSize = sizeof(ID3D11Texture2D*);
 		ap.AllocAlignment = alignof(ID3D11Texture2D*);
@@ -1645,12 +1614,10 @@ bool FVideoDecoderH264_WindowsImpl::AllocateApplicationOutputSampleBuffer(int32
 		ap.Width = Width;
 		ap.Height = Height;
 		ap.BytesPerPixel = 1;
-		PRAGMA_ENABLE_DEPRECATION_WARNINGS
 	}
 	else if (DecodeMode == EDecodeMode::Win10HW || DecodeMode == EDecodeMode::Win10SW || DecodeMode == EDecodeMode::XB1)
 	{
 		// We need a byte buffer to put the image into. On the application side this may be realized as a TArray<uint8>.
-		PRAGMA_DISABLE_DEPRECATION_WARNINGS
 		ap.FrameBufferType = EFrameBufferType::CODEC_RawBuffer;
 		ap.AllocSize = Width * Height * 3 / 2;
 		ap.AllocAlignment = 1;
@@ -1658,12 +1625,10 @@ bool FVideoDecoderH264_WindowsImpl::AllocateApplicationOutputSampleBuffer(int32
 		ap.Width = Width;
 		ap.Height = Height;
 		ap.BytesPerPixel = 1;
-		PRAGMA_ENABLE_DEPRECATION_WARNINGS
 	}
 	else if (DecodeMode == EDecodeMode::Win8SW)
 	{
 		// We want to have an ID3D11Texture2D we can give the decoder to decode into.
-		PRAGMA_DISABLE_DEPRECATION_WARNINGS
 		ap.FrameBufferType = EFrameBufferType::CODEC_TextureObject;
 		ap.AllocSize = sizeof(ID3D11Texture2D*);
 		ap.AllocAlignment = alignof(ID3D11Texture2D*);
@@ -1671,7 +1636,6 @@ bool FVideoDecoderH264_WindowsImpl::AllocateApplicationOutputSampleBuffer(int32
 		ap.Width = Width;
 		ap.Height = Height;
 		ap.BytesPerPixel = 1;
-		PRAGMA_ENABLE_DEPRECATION_WARNINGS
 	}
 	else
 	{
@@ -1679,20 +1643,14 @@ bool FVideoDecoderH264_WindowsImpl::AllocateApplicationOutputSampleBuffer(int32
 		return false;
 	}
 
-	PRAGMA_DISABLE_DEPRECATION_WARNINGS
 	ar = AllocateOutputFrameBuffer(CurrentApplicationOutputSample->GetBuffer(), &ap);
 	if (ar == EFrameBufferAllocReturn::CODEC_Success)
-	PRAGMA_ENABLE_DEPRECATION_WARNINGS
 	{
 		// Got an output buffer.
-		PRAGMA_DISABLE_DEPRECATION_WARNINGS
 		check(CurrentApplicationOutputSample->GetBuffer()->AllocatedBuffer);
-		PRAGMA_ENABLE_DEPRECATION_WARNINGS
 		return true;
 	}
-	PRAGMA_DISABLE_DEPRECATION_WARNINGS
 	else if (ar == EFrameBufferAllocReturn::CODEC_TryAgainLater)
-	PRAGMA_ENABLE_DEPRECATION_WARNINGS
 	{
 		UE_LOG(LogVideoDecoder, Error, TEXT("Unsupported allocation return value"));
 		return false;
@@ -1708,12 +1666,10 @@ void FVideoDecoderH264_WindowsImpl::ReturnUnusedApplicationOutputSample()
 {
 	if (CurrentApplicationOutputSample)
 	{
-		PRAGMA_DISABLE_DEPRECATION_WARNINGS
 		if (FVideoDecoderAllocFrameBufferResult* afb = CurrentApplicationOutputSample->GetBuffer())
 		{
 			afb->ReleaseCallback.ExecuteIfBound(afb->CallbackValue, afb->AllocatedBuffer);
 		}
-		PRAGMA_ENABLE_DEPRECATION_WARNINGS
 		CurrentApplicationOutputSample->Release();
 		CurrentApplicationOutputSample = nullptr;
 	}
@@ -1792,9 +1748,7 @@ bool FVideoDecoderH264_WindowsImpl::ProcessDecoding(FInputToDecode* InInput)
 			if (DecodeMode == EDecodeMode::Win8SW || DecodeMode == EDecodeMode::Win10SW)
 			{
 				// DX11 software mode needs a texture prepared by the application
-				PRAGMA_DISABLE_DEPRECATION_WARNINGS
 				if (DecodeMode == EDecodeMode::Win8SW && !CurrentApplicationOutputSample->GetBuffer()->AllocatedBuffer)
-				PRAGMA_ENABLE_DEPRECATION_WARNINGS
 				{
 					if (!AllocateApplicationOutputSampleBuffer(InWidth, InHeight))
 					{
@@ -1805,9 +1759,7 @@ bool FVideoDecoderH264_WindowsImpl::ProcessDecoding(FInputToDecode* InInput)
 				{
 					CurrentSoftwareDecoderOutput = new FF5PlayerVideoDecoderOutputDX;
 				}
-				PRAGMA_DISABLE_DEPRECATION_WARNINGS
 				if (!CurrentSoftwareDecoderOutput->PreInitForSoftwareDecode(FIntPoint(InWidth, InHeight), DecodeMode == EDecodeMode::Win8SW ? *static_cast<ID3D11Texture2D**>(CurrentApplicationOutputSample->GetBuffer()->AllocatedBuffer) : nullptr))
-				PRAGMA_ENABLE_DEPRECATION_WARNINGS
 				{
 					return false;
 				}
@@ -1912,9 +1864,7 @@ bool FVideoDecoderH264_WindowsImpl::CopyTexture(const TRefCountPtr<IMFSample>& D
 
 	if (DecodeMode == EDecodeMode::Win8HW)
 	{
-		PRAGMA_DISABLE_DEPRECATION_WARNINGS
 		check(CurrentApplicationOutputSample->GetBuffer()->AllocatedBuffer == nullptr);
-		PRAGMA_ENABLE_DEPRECATION_WARNINGS
 		if (!AllocateApplicationOutputSampleBuffer(OutputDim.X, OutputDim.Y))
 		{
 			UE_LOG(LogVideoDecoder, Error, TEXT("Failed to get a decode output buffer from the application!"));
@@ -1924,16 +1874,12 @@ bool FVideoDecoderH264_WindowsImpl::CopyTexture(const TRefCountPtr<IMFSample>& D
 	}
 	else if (DecodeMode == EDecodeMode::Win8SW)
 	{
-		PRAGMA_DISABLE_DEPRECATION_WARNINGS
 		check(CurrentApplicationOutputSample->GetBuffer()->AllocatedBuffer != nullptr);
-		PRAGMA_ENABLE_DEPRECATION_WARNINGS
 		DecoderOutput->InitializeWithTextureBuffer(OutputDim.X, FIntPoint(OutputDim.X, OutputDim.Y));// * 3 / 2));
 	}
 	else if (DecodeMode == EDecodeMode::Win10SW)
 	{
-		PRAGMA_DISABLE_DEPRECATION_WARNINGS
 		check(CurrentApplicationOutputSample->GetBuffer()->AllocatedBuffer == nullptr);
-		PRAGMA_ENABLE_DEPRECATION_WARNINGS
 		if (!AllocateApplicationOutputSampleBuffer(NativeDim.X, NativeDim.Y))
 		{
 			UE_LOG(LogVideoDecoder, Error, TEXT("Failed to get a decode output buffer from the application!"));
@@ -1960,10 +1906,8 @@ bool FVideoDecoderH264_WindowsImpl::CopyTexture(const TRefCountPtr<IMFSample>& D
 
 		// Copy NV12 texture data into external application buffer
 		check(CurrentApplicationOutputSample);
-		PRAGMA_DISABLE_DEPRECATION_WARNINGS
 		check((DWORD)CurrentApplicationOutputSample->GetBuffer()->AllocatedSize >= BufferSize);
 		FMemory::Memcpy(CurrentApplicationOutputSample->GetBuffer()->AllocatedBuffer, Data, BufferSize);
-		PRAGMA_ENABLE_DEPRECATION_WARNINGS
 
 		Result = Buffer->Unlock();
 		if (FAILED(Result))
@@ -2015,9 +1959,7 @@ bool FVideoDecoderH264_WindowsImpl::CopyTexture(const TRefCountPtr<IMFSample>& D
 			return false;
 		}
 
-		PRAGMA_DISABLE_DEPRECATION_WARNINGS
 		check(CurrentApplicationOutputSample->GetBuffer()->AllocatedBuffer == nullptr);
-		PRAGMA_ENABLE_DEPRECATION_WARNINGS
 		if (!AllocateApplicationOutputSampleBuffer(Pitch, TextureDesc.Height))
 		{
 			UE_LOG(LogVideoDecoder, Error, TEXT("Failed to get a decode output buffer from the application!"));
@@ -2032,10 +1974,8 @@ bool FVideoDecoderH264_WindowsImpl::CopyTexture(const TRefCountPtr<IMFSample>& D
 		DecoderOutput->InitializeWithBuffer(Pitch, FIntPoint(TextureDesc.Width, TextureDesc.Height * 3 / 2));
 
 		// Copy NV12 texture data into external application buffer
-		PRAGMA_DISABLE_DEPRECATION_WARNINGS
 		check((DWORD)CurrentApplicationOutputSample->GetBuffer()->AllocatedSize >= BufferSize);
 		FMemory::Memcpy(CurrentApplicationOutputSample->GetBuffer()->AllocatedBuffer, Data, BufferSize);
-		PRAGMA_ENABLE_DEPRECATION_WARNINGS
 
 		Result = Buffer2D->Unlock2D();
 		if (FAILED(Result))
@@ -2110,9 +2050,7 @@ bool FVideoDecoderH264_WindowsImpl::ConvertDecodedImage(const TRefCountPtr<IMFSa
 	CurrentApplicationOutputSample->SetNativeDecoderOutput(DecoderOutput);
 	if (DecodeMode == EDecodeMode::Win8HW)
 	{
-		PRAGMA_DISABLE_DEPRECATION_WARNINGS
 		ID3D11Texture2D** ApplicationTexturePtr = reinterpret_cast<ID3D11Texture2D**>(CurrentApplicationOutputSample->GetBuffer()->AllocatedBuffer);
-		PRAGMA_ENABLE_DEPRECATION_WARNINGS
 		// We set the texture in the user provided buffer. Note that we do NOT add the ref count here. This is ok because we have the
 		// DecoderOutput stored in CurrentApplicationOutputSample where it is held. If the user application code uses the texture it
 		// will/shall attach it to a ref counter pointer to retain it.
@@ -2203,12 +2141,10 @@ void FVideoDecoderH264_WindowsImpl::PurgeAllPendingOutput()
 	FVideoDecoderOutputH264_Windows* Output = nullptr;
 	while(OutputQueue.Dequeue(Output))
 	{
-		PRAGMA_DISABLE_DEPRECATION_WARNINGS
 		if (FVideoDecoderAllocFrameBufferResult* afb = Output->GetBuffer())
 		{
 			afb->ReleaseCallback.ExecuteIfBound(afb->CallbackValue, afb->AllocatedBuffer);
 		}
-		PRAGMA_ENABLE_DEPRECATION_WARNINGS
 		Output->Release();
 		Output = nullptr;
 	}
@@ -2250,15 +2186,11 @@ FVideoDecoderH264_WindowsImpl::~FVideoDecoderH264_WindowsImpl()
 /**
  * Called right after the constructor with initialization parameters.
  */
-PRAGMA_DISABLE_DEPRECATION_WARNINGS
 bool FVideoDecoderH264_WindowsImpl::Setup(const FInit& InInit)
-PRAGMA_ENABLE_DEPRECATION_WARNINGS
 {
 	// Remember the decoder allocation interface factory methods.
-	PRAGMA_DISABLE_DEPRECATION_WARNINGS
 	CreateDecoderAllocationInterfaceFN = InInit.CreateDecoderAllocationInterface;
 	ReleaseDecoderAllocationInterfaceFN = InInit.ReleaseDecoderAllocationInterface;
-	PRAGMA_ENABLE_DEPRECATION_WARNINGS
 	// [...]
 	CreateWorkerThread();
 	return true;
@@ -2272,9 +2204,7 @@ void FVideoDecoderH264_WindowsImpl::PerformAsyncShutdown()
 	if (bIsInitialized)
 	{
 		ReleaseD3DResources();
-		PRAGMA_DISABLE_DEPRECATION_WARNINGS
 		ReleaseDecoderAllocationInterface();
-		PRAGMA_ENABLE_DEPRECATION_WARNINGS
 		StaticReleaseResources();
 		bIsInitialized = false;
 	}
@@ -2301,9 +2231,7 @@ bool FVideoDecoderH264_WindowsImpl::SetupFirstUseResources()
 	if (!bIsInitialized)
 	{
 		if (StaticInitializeResources() &&
-			PRAGMA_DISABLE_DEPRECATION_WARNINGS
 			CreateDecoderAllocationInterface() &&
-			PRAGMA_ENABLE_DEPRECATION_WARNINGS
 			CreateD3DResources())
 		{
 			bIsInitialized = true;
@@ -2317,31 +2245,25 @@ bool FVideoDecoderH264_WindowsImpl::SetupFirstUseResources()
 /**
  * Called from webrtc indirectly to decode a frame.
  */
-PRAGMA_DISABLE_DEPRECATION_WARNINGS
 FVideoDecoder::EDecodeResult FVideoDecoderH264_WindowsImpl::Decode(const FVideoDecoderInput* InInput)
-PRAGMA_ENABLE_DEPRECATION_WARNINGS
 {
 	bool bOk = SetupFirstUseResources();
 	if (bOk)
 	{
-		PRAGMA_DISABLE_DEPRECATION_WARNINGS
 		if (!InInput->GetData() || InInput->GetDataSize() <= 0)
 		{
 			return FVideoDecoder::EDecodeResult::Failure;
 		}
-		PRAGMA_ENABLE_DEPRECATION_WARNINGS
 
 		// Need to wait for the decoder thread to have spun up??
 
 		TUniquePtr<FInputToDecode> Input = MakeUnique<FInputToDecode>();
-		PRAGMA_DISABLE_DEPRECATION_WARNINGS
 		Input->PTS = InInput->GetPTS();
 		Input->Width = InInput->GetWidth();
 		Input->Height = InInput->GetHeight();
 		Input->bIsKeyframe = InInput->IsKeyframe();
 		Input->Data.AddUninitialized(InInput->GetDataSize());
 		FMemory::Memcpy(Input->Data.GetData(), InInput->GetData(), InInput->GetDataSize());
-		PRAGMA_ENABLE_DEPRECATION_WARNINGS
 		InputQueue.Enqueue(MoveTemp(Input));
 
 		WorkerThreadSignalHaveWork.Signal();
@@ -2352,29 +2274,18 @@ PRAGMA_ENABLE_DEPRECATION_WARNINGS
 		{
 			if (Output->GetIsValid())
 			{
-				PRAGMA_DISABLE_DEPRECATION_WARNINGS
 				OnDecodedFrame(Output);
-				PRAGMA_ENABLE_DEPRECATION_WARNINGS
 			}
-			
 			else
 			{
 				delete Output;
-				PRAGMA_DISABLE_DEPRECATION_WARNINGS
 				return FVideoDecoder::EDecodeResult::Failure;
-				PRAGMA_ENABLE_DEPRECATION_WARNINGS
 			}
-			
 			Output = nullptr;
 		}
-		PRAGMA_DISABLE_DEPRECATION_WARNINGS
 		return FVideoDecoder::EDecodeResult::Success;
-		PRAGMA_ENABLE_DEPRECATION_WARNINGS
 	}
-	
-	PRAGMA_DISABLE_DEPRECATION_WARNINGS
 	return FVideoDecoder::EDecodeResult::Failure;
-	PRAGMA_ENABLE_DEPRECATION_WARNINGS
 }
 
 
diff --git a/Engine/Source/Runtime/AVEncoder/Private/Decoders/Windows/VideoDecoderH264_Windows.h b/Engine/Source/Runtime/AVEncoder/Private/Decoders/Windows/VideoDecoderH264_Windows.h
--- a/Engine/Source/Runtime/AVEncoder/Private/Decoders/Windows/VideoDecoderH264_Windows.h
+++ b/Engine/Source/Runtime/AVEncoder/Private/Decoders/Windows/VideoDecoderH264_Windows.h
@@ -12,14 +12,12 @@ class FVideoDecoderH264_Windows : public FVideoDecoderH264_Impl
 {
 public:
 	// register decoder with video decoder factory
-	PRAGMA_DISABLE_DEPRECATION_WARNINGS
 	static void Register(FVideoDecoderFactory& InFactory);
-	virtual bool Setup(const FInit& InInit) = 0;
-	virtual EDecodeResult Decode(const FVideoDecoderInput* InInput) = 0;
-	PRAGMA_ENABLE_DEPRECATION_WARNINGS
 
+	virtual bool Setup(const FInit& InInit) = 0;
 	virtual void Shutdown() = 0;
 
+	virtual EDecodeResult Decode(const FVideoDecoderInput* InInput) = 0;
 protected:
 	virtual ~FVideoDecoderH264_Windows() = default;
 };
diff --git a/Engine/Source/Runtime/AVEncoder/Private/Encoders/WmfAudioEncoder.cpp b/Engine/Source/Runtime/AVEncoder/Private/Encoders/WmfAudioEncoder.cpp
--- a/Engine/Source/Runtime/AVEncoder/Private/Encoders/WmfAudioEncoder.cpp
+++ b/Engine/Source/Runtime/AVEncoder/Private/Encoders/WmfAudioEncoder.cpp
@@ -12,9 +12,7 @@
 
 namespace AVEncoder
 {
-PRAGMA_DISABLE_DEPRECATION_WARNINGS
 class FWmfAudioEncoder : public FAudioEncoder
-PRAGMA_ENABLE_DEPRECATION_WARNINGS
 {
 public:
 
@@ -28,12 +26,10 @@ public:
 
 	const TCHAR* GetName() const override;
 	const TCHAR* GetType() const override;
-	PRAGMA_DISABLE_DEPRECATION_WARNINGS
 	bool Initialize(const FAudioConfig& InConfig) override;
+	void Shutdown() override;
 	void Encode(const FAudioFrame& Frame) override;
 	FAudioConfig GetConfig() const override;
-	PRAGMA_ENABLE_DEPRECATION_WARNINGS
-	void Shutdown() override;
 
 private:
 	bool SetInputType();
@@ -41,18 +37,14 @@ private:
 	bool RetrieveStreamInfo();
 	bool StartStreaming();
 
-	PRAGMA_DISABLE_DEPRECATION_WARNINGS
 	FIMFSampleWrapper CreateInputSample(const uint8* SampleData, uint32 Size, FTimespan Timestamp, FTimespan Duration);
 	FIMFSampleWrapper GetOutputSample();
-	PRAGMA_ENABLE_DEPRECATION_WARNINGS
 
 	ECodecType CodecType;
 	FString Name;
 	FString Type;
 
-	PRAGMA_DISABLE_DEPRECATION_WARNINGS
 	FAudioConfig Config;
-	PRAGMA_ENABLE_DEPRECATION_WARNINGS
 	TRefCountPtr<IMFTransform> Encoder;
 	TRefCountPtr<IMFMediaType> OutputType;
 	MFT_INPUT_STREAM_INFO InputStreamInfo = {};
@@ -66,7 +58,7 @@ private:
 //////////////////////////////////////////////////////////////////////////
 // FWmfAudioEncoder implementation
 //////////////////////////////////////////////////////////////////////////
-PRAGMA_DISABLE_DEPRECATION_WARNINGS
+
 FWmfAudioEncoder::FWmfAudioEncoder(ECodecType CodecType)
 	: CodecType(CodecType)
 {
@@ -80,13 +72,10 @@ FWmfAudioEncoder::FWmfAudioEncoder(ECodecType CodecType)
 		checkNoEntry();
 	}
 }
-PRAGMA_ENABLE_DEPRECATION_WARNINGS
 
-PRAGMA_DISABLE_DEPRECATION_WARNINGS
 FWmfAudioEncoder::~FWmfAudioEncoder()
 {
 }
-PRAGMA_ENABLE_DEPRECATION_WARNINGS
 
 const TCHAR* FWmfAudioEncoder::GetName() const
 {
@@ -127,9 +116,7 @@ namespace
 	}
 }
 
-PRAGMA_DISABLE_DEPRECATION_WARNINGS
 bool FWmfAudioEncoder::Initialize(const FAudioConfig& InConfig)
-PRAGMA_ENABLE_DEPRECATION_WARNINGS
 {
 
 	//
@@ -139,12 +126,10 @@ PRAGMA_ENABLE_DEPRECATION_WARNINGS
 	{
 		// See  https://docs.microsoft.com/en-us/windows/desktop/medfound/aac-encoder for details
 		FString ErrorStr;
-		PRAGMA_DISABLE_DEPRECATION_WARNINGS
 		if (
 			!ValidateValue(TEXT("AAC Bitrate"), InConfig.Bitrate, {12000*8, 16000*8, 20000*8, 24000*8}, ErrorStr) ||
 			!ValidateValue(TEXT("AAC Samplerate"), InConfig.Samplerate, { 44100, 48000, 0 }, ErrorStr) ||
 			!ValidateValue(TEXT("AAC NumChannels"), InConfig.NumChannels, {1,2,6}, ErrorStr))
-		PRAGMA_ENABLE_DEPRECATION_WARNINGS
 		{
 			UE_LOG(LogAVEncoder, Error, TEXT("%s"), *ErrorStr);
 			return false;
@@ -156,11 +141,9 @@ PRAGMA_ENABLE_DEPRECATION_WARNINGS
 		return false;
 	}
 
-	PRAGMA_DISABLE_DEPRECATION_WARNINGS
 	UE_LOG(LogAVEncoder, Log, TEXT("AudioEncoder config: %d channels, %d Hz, %.2f Kbps"), InConfig.NumChannels, InConfig.Samplerate, InConfig.Bitrate / 1000.0f);
 
 	Config = InConfig;
-	PRAGMA_ENABLE_DEPRECATION_WARNINGS
 
 	const GUID* CodecGuid=nullptr;
 	if (CodecType == ECodecType::AAC)
@@ -188,7 +171,6 @@ PRAGMA_ENABLE_DEPRECATION_WARNINGS
 	return true;
 }
 
-
 bool FWmfAudioEncoder::SetInputType()
 {
 	TRefCountPtr<IMFMediaType> MediaType;
@@ -196,11 +178,8 @@ bool FWmfAudioEncoder::SetInputType()
 	CHECK_HR(MediaType->SetGUID(MF_MT_MAJOR_TYPE, MFMediaType_Audio));
 	CHECK_HR(MediaType->SetGUID(MF_MT_SUBTYPE, MFAudioFormat_PCM));
 	CHECK_HR(MediaType->SetUINT32(MF_MT_AUDIO_BITS_PER_SAMPLE, 16)); // the only value supported
-
-	PRAGMA_DISABLE_DEPRECATION_WARNINGS
 	CHECK_HR(MediaType->SetUINT32(MF_MT_AUDIO_SAMPLES_PER_SECOND, Config.Samplerate));
 	CHECK_HR(MediaType->SetUINT32(MF_MT_AUDIO_NUM_CHANNELS, Config.NumChannels));
-	PRAGMA_DISABLE_DEPRECATION_WARNINGS
 
 	CHECK_HR(Encoder->SetInputType(0, MediaType, 0));
 
diff --git a/Engine/Source/Runtime/AVEncoder/Private/Encoders/WmfAudioEncoder.h b/Engine/Source/Runtime/AVEncoder/Private/Encoders/WmfAudioEncoder.h
--- a/Engine/Source/Runtime/AVEncoder/Private/Encoders/WmfAudioEncoder.h
+++ b/Engine/Source/Runtime/AVEncoder/Private/Encoders/WmfAudioEncoder.h
@@ -10,18 +10,14 @@
 
 namespace AVEncoder
 {
-	PRAGMA_DISABLE_DEPRECATION_WARNINGS
 	class FWmfAudioEncoderFactory : public FAudioEncoderFactory
-	PRAGMA_ENABLE_DEPRECATION_WARNINGS
 	{
 	public:
 		FWmfAudioEncoderFactory();
 		~FWmfAudioEncoderFactory() override;
 		const TCHAR* GetName() const override;
 		TArray<FString> GetSupportedCodecs() const override;
-		PRAGMA_DISABLE_DEPRECATION_WARNINGS
 		TUniquePtr<FAudioEncoder> CreateEncoder(const FString& Codec) override;
-		PRAGMA_ENABLE_DEPRECATION_WARNINGS
 	private:
 	};
 }
diff --git a/Engine/Source/Runtime/AVEncoder/Private/VideoDecoderFactory.cpp b/Engine/Source/Runtime/AVEncoder/Private/VideoDecoderFactory.cpp
--- a/Engine/Source/Runtime/AVEncoder/Private/VideoDecoderFactory.cpp
+++ b/Engine/Source/Runtime/AVEncoder/Private/VideoDecoderFactory.cpp
@@ -23,13 +23,11 @@ namespace AVEncoder
 FThreadSafeCounter			FVideoDecoderFactory::NextID = 1;
 bool						FVideoDecoderFactory::bDebugDontRegisterDefaultCodecs = false;
 
-PRAGMA_DISABLE_DEPRECATION_WARNINGS
 FVideoDecoderFactory& FVideoDecoderFactory::Get()
 {
 	static FVideoDecoderFactory Singleton;
 	return Singleton;
 }
-PRAGMA_ENABLE_DEPRECATION_WARNINGS
 
 FVideoDecoderFactory::FVideoDecoderFactory()
 {
@@ -52,14 +50,10 @@ void FVideoDecoderFactory::Debug_SetDontRegisterDefaultCodecs()
 }
 
 
-PRAGMA_DISABLE_DEPRECATION_WARNINGS
 void FVideoDecoderFactory::Register(const FVideoDecoderInfo& InInfo, const CreateDecoderCallback& InCreateEncoder)
-PRAGMA_ENABLE_DEPRECATION_WARNINGS
 {
 	AvailableDecoders.Push(InInfo);
-	PRAGMA_DISABLE_DEPRECATION_WARNINGS
 	AvailableDecoders.Last().ID = NextID.Increment();
-	PRAGMA_ENABLE_DEPRECATION_WARNINGS
 	CreateDecoders.Push(InCreateEncoder);
 }
 
@@ -81,15 +75,11 @@ void FVideoDecoderFactory::RegisterDefaultCodecs()
 #endif
 }
 
-PRAGMA_DISABLE_DEPRECATION_WARNINGS
 bool FVideoDecoderFactory::GetInfo(uint32 InID, FVideoDecoderInfo& OutInfo) const
-PRAGMA_ENABLE_DEPRECATION_WARNINGS
 {
 	for (int32 Index = 0; Index < AvailableDecoders.Num(); ++Index)
 	{
-		PRAGMA_DISABLE_DEPRECATION_WARNINGS
 		if (AvailableDecoders[Index].ID == InID)
-		PRAGMA_ENABLE_DEPRECATION_WARNINGS
 		{
 			OutInfo = AvailableDecoders[Index];
 			return true;
@@ -98,27 +88,20 @@ PRAGMA_ENABLE_DEPRECATION_WARNINGS
 	return false;
 }
 
-PRAGMA_DISABLE_DEPRECATION_WARNINGS
+
 FVideoDecoder* FVideoDecoderFactory::Create(uint32 InID, const FVideoDecoder::FInit& InInit)
-PRAGMA_ENABLE_DEPRECATION_WARNINGS
 {
-	PRAGMA_DISABLE_DEPRECATION_WARNINGS
 	FVideoDecoder* Result = nullptr;
-	PRAGMA_ENABLE_DEPRECATION_WARNINGS
 	for (int32 Index = 0; Index < AvailableDecoders.Num(); ++Index)
 	{
-		PRAGMA_DISABLE_DEPRECATION_WARNINGS
 		if (AvailableDecoders[Index].ID == InID)
-		PRAGMA_ENABLE_DEPRECATION_WARNINGS
 		{
 			Result = CreateDecoders[Index]();
-			PRAGMA_DISABLE_DEPRECATION_WARNINGS
 			if (Result && !Result->Setup(InInit))
 			{
 				Result->Shutdown();
 				Result = nullptr;
 			}
-			PRAGMA_ENABLE_DEPRECATION_WARNINGS
 			break;
 		}
 	}
diff --git a/Engine/Source/Runtime/AVEncoder/Private/VideoDecoderInput.cpp b/Engine/Source/Runtime/AVEncoder/Private/VideoDecoderInput.cpp
--- a/Engine/Source/Runtime/AVEncoder/Private/VideoDecoderInput.cpp
+++ b/Engine/Source/Runtime/AVEncoder/Private/VideoDecoderInput.cpp
@@ -5,15 +5,11 @@
 namespace AVEncoder
 {
 
-PRAGMA_DISABLE_DEPRECATION_WARNINGS
 class FVideoDecoderInputImpl : public FVideoDecoderInput
 {
-PRAGMA_ENABLE_DEPRECATION_WARNINGS
 public:
-	PRAGMA_DISABLE_DEPRECATION_WARNINGS
 	virtual ~FVideoDecoderInputImpl() = default;
 	FVideoDecoderInputImpl(const FInputData& InInputData);
-	PRAGMA_ENABLE_DEPRECATION_WARNINGS
 
 	virtual int32 GetWidth() const override
 	{ return Width; }
@@ -49,14 +45,12 @@ private:
 	bool		bMissingFrames;
 };
 
-PRAGMA_DISABLE_DEPRECATION_WARNINGS
 TSharedPtr<FVideoDecoderInput> FVideoDecoderInput::Create(const FInputData& InInputData)
 {
 	return MakeShared<FVideoDecoderInputImpl>(InInputData);
 }
-PRAGMA_ENABLE_DEPRECATION_WARNINGS
 
-PRAGMA_DISABLE_DEPRECATION_WARNINGS
+
 FVideoDecoderInputImpl::FVideoDecoderInputImpl(const FInputData& InInputData)
 {
 	AccessUnit = InInputData.EncodedData;
@@ -70,6 +64,6 @@ FVideoDecoderInputImpl::FVideoDecoderInputImpl(const FInputData& InInputData)
 	bIsCompleteFrame = InInputData.bIsComplete;
 	bMissingFrames = InInputData.bMissingFrames;
 }
-PRAGMA_ENABLE_DEPRECATION_WARNINGS
+
 
 } /* namespace AVEncoder */
diff --git a/Engine/Source/Runtime/AVEncoder/Private/VideoEncoderInput.cpp b/Engine/Source/Runtime/AVEncoder/Private/VideoEncoderInput.cpp
--- a/Engine/Source/Runtime/AVEncoder/Private/VideoEncoderInput.cpp
+++ b/Engine/Source/Runtime/AVEncoder/Private/VideoEncoderInput.cpp
@@ -30,9 +30,8 @@ namespace AVEncoder
 	// *** FVideoEncoderInput *************************************************************************
 
 	// --- construct video encoder input based on expected input frame format -------------------------
-	PRAGMA_DISABLE_DEPRECATION_WARNINGS
+
 	TSharedPtr<FVideoEncoderInput> FVideoEncoderInput::CreateDummy(bool bIsResizable)
-	PRAGMA_ENABLE_DEPRECATION_WARNINGS
 	{
 		TSharedPtr<FVideoEncoderInputImpl>	Input = MakeShared<FVideoEncoderInputImpl>();
 		Input->bIsResizable = bIsResizable;
@@ -44,9 +43,7 @@ namespace AVEncoder
 		return Input;
 	}
 
-	PRAGMA_DISABLE_DEPRECATION_WARNINGS
 	TSharedPtr<FVideoEncoderInput> FVideoEncoderInput::CreateForYUV420P(uint32 InWidth, uint32 InHeight, bool bIsResizable)
-	PRAGMA_ENABLE_DEPRECATION_WARNINGS
 	{
 		TSharedPtr<FVideoEncoderInputImpl>	Input = MakeShared<FVideoEncoderInputImpl>();
 		Input->bIsResizable = bIsResizable;
@@ -58,9 +55,7 @@ namespace AVEncoder
 		return Input;
 	}
 
-	PRAGMA_DISABLE_DEPRECATION_WARNINGS
 	TSharedPtr<FVideoEncoderInput> FVideoEncoderInput::CreateForD3D11(void* InApplicationD3DDevice, bool bIsResizable, bool IsShared)
-	PRAGMA_ENABLE_DEPRECATION_WARNINGS
 	{
 		TSharedPtr<FVideoEncoderInputImpl>	Input = MakeShared<FVideoEncoderInputImpl>();
 
@@ -88,9 +83,7 @@ namespace AVEncoder
 	}
 
 
-	PRAGMA_DISABLE_DEPRECATION_WARNINGS
 	TSharedPtr<FVideoEncoderInput> FVideoEncoderInput::CreateForD3D12(void* InApplicationD3DDevice, bool bIsResizable, bool IsShared)
-	PRAGMA_ENABLE_DEPRECATION_WARNINGS
 	{
 		TSharedPtr<FVideoEncoderInputImpl>	Input = MakeShared<FVideoEncoderInputImpl>();
 
@@ -117,9 +110,7 @@ namespace AVEncoder
 		return Input;
 	}
 
-	PRAGMA_DISABLE_DEPRECATION_WARNINGS
 	TSharedPtr<FVideoEncoderInput> FVideoEncoderInput::CreateForCUDA(void* InApplicationContext, bool bIsResizable)
-	PRAGMA_ENABLE_DEPRECATION_WARNINGS
 	{
 		TSharedPtr<FVideoEncoderInputImpl> Input = MakeShared<FVideoEncoderInputImpl>();
 		Input->bIsResizable = bIsResizable;
@@ -132,20 +123,14 @@ namespace AVEncoder
 	}
 
 #if PLATFORM_DESKTOP && !PLATFORM_APPLE
-	PRAGMA_DISABLE_DEPRECATION_WARNINGS
 	TSharedPtr<FVideoEncoderInput> FVideoEncoderInput::CreateForVulkan(void* InApplicationVulkanData, bool bIsResizable)
-	PRAGMA_ENABLE_DEPRECATION_WARNINGS
 	{
 		TSharedPtr<FVideoEncoderInputImpl>	Input = MakeShared<FVideoEncoderInputImpl>();
 		Input->bIsResizable = bIsResizable;
 
-		PRAGMA_DISABLE_DEPRECATION_WARNINGS
 		FVulkanDataStruct* VulkanData = static_cast<FVulkanDataStruct*>(InApplicationVulkanData);
-		PRAGMA_ENABLE_DEPRECATION_WARNINGS
 
-		PRAGMA_DISABLE_DEPRECATION_WARNINGS
 		if (!Input->SetupForVulkan(VulkanData->VulkanInstance, VulkanData->VulkanPhysicalDevice, VulkanData->VulkanDevice))
-		PRAGMA_ENABLE_DEPRECATION_WARNINGS
 		{
 			Input.Reset();
 		}
@@ -166,7 +151,6 @@ namespace AVEncoder
 
 	// *** FVideoEncoderInputImpl *********************************************************************
 
-	PRAGMA_DISABLE_DEPRECATION_WARNINGS
 	FVideoEncoderInputImpl::~FVideoEncoderInputImpl()
 	{
 		{
@@ -182,21 +166,16 @@ namespace AVEncoder
 		//	DEBUG_D3D11_REPORT_LIVE_DEVICE_OBJECT(FrameInfoD3D.EncoderDeviceD3D11);
 #endif
 	}
-	PRAGMA_ENABLE_DEPRECATION_WARNINGS
 
 	bool FVideoEncoderInputImpl::SetupForDummy()
 	{
-		PRAGMA_DISABLE_DEPRECATION_WARNINGS
 		FrameFormat = EVideoFrameFormat::Undefined;
-		PRAGMA_ENABLE_DEPRECATION_WARNINGS
 		return true;
 	}
 
 	bool FVideoEncoderInputImpl::SetupForYUV420P(uint32 InWidth, uint32 InHeight)
 	{
-		PRAGMA_DISABLE_DEPRECATION_WARNINGS
 		FrameFormat = EVideoFrameFormat::YUV420P;
-		PRAGMA_ENABLE_DEPRECATION_WARNINGS
 		FrameInfoYUV420P.StrideY = InWidth;
 		FrameInfoYUV420P.StrideU = (InWidth + 1) / 2;
 		FrameInfoYUV420P.StrideV = (InWidth + 1) / 2;
@@ -245,9 +224,7 @@ namespace AVEncoder
 		DEBUG_SET_D3D11_OBJECT_NAME(FrameInfoD3D.EncoderDeviceD3D11, "FVideoEncoderInputImpl");
 		DEBUG_SET_D3D11_OBJECT_NAME(FrameInfoD3D.EncoderDeviceContextD3D11, "FVideoEncoderInputImpl");
 
-		PRAGMA_DISABLE_DEPRECATION_WARNINGS
 		FrameFormat = EVideoFrameFormat::D3D11_R8G8B8A8_UNORM;
-		PRAGMA_ENABLE_DEPRECATION_WARNINGS
 
 		CollectAvailableEncoders();
 		return true;
@@ -297,9 +274,7 @@ namespace AVEncoder
 		DEBUG_SET_D3D11_OBJECT_NAME(FrameInfoD3D.EncoderDeviceD3D11, "FVideoEncoderInputImpl");
 		DEBUG_SET_D3D11_OBJECT_NAME(FrameInfoD3D.EncoderDeviceContextD3D11, "FVideoEncoderInputImpl");
 
-		PRAGMA_DISABLE_DEPRECATION_WARNINGS
 		FrameFormat = EVideoFrameFormat::D3D11_R8G8B8A8_UNORM;
-		PRAGMA_ENABLE_DEPRECATION_WARNINGS
 
 		CollectAvailableEncoders();
 		return true;
@@ -336,9 +311,7 @@ namespace AVEncoder
 			return false;
 		}
 
-		PRAGMA_DISABLE_DEPRECATION_WARNINGS
 		FrameFormat = EVideoFrameFormat::D3D12_R8G8B8A8_UNORM;
-		PRAGMA_ENABLE_DEPRECATION_WARNINGS
 
 		CollectAvailableEncoders();
 		return true;
@@ -398,9 +371,7 @@ namespace AVEncoder
 		DEBUG_SET_D3D11_OBJECT_NAME(FrameInfoD3D.EncoderDeviceD3D11, "FVideoEncoderInputImpl");
 		DEBUG_SET_D3D11_OBJECT_NAME(FrameInfoD3D.EncoderDeviceContextD3D11, "FVideoEncoderInputImpl");
 
-		PRAGMA_DISABLE_DEPRECATION_WARNINGS
 		FrameFormat = EVideoFrameFormat::D3D11_R8G8B8A8_UNORM;
-		PRAGMA_ENABLE_DEPRECATION_WARNINGS
 
 		CollectAvailableEncoders();
 		return true;
@@ -415,9 +386,7 @@ namespace AVEncoder
 	{
 		FrameInfoCUDA.EncoderContextCUDA = static_cast<CUcontext>(InApplicationContext);
 
-		PRAGMA_DISABLE_DEPRECATION_WARNINGS
 		FrameFormat = EVideoFrameFormat::CUDA_R8G8B8A8_UNORM;
-		PRAGMA_ENABLE_DEPRECATION_WARNINGS
 
 		CollectAvailableEncoders();
 		return true;
@@ -426,13 +395,11 @@ namespace AVEncoder
 #if PLATFORM_DESKTOP && !PLATFORM_APPLE
 	bool FVideoEncoderInputImpl::SetupForVulkan(VkInstance InApplicationVulkanInstance, VkPhysicalDevice InApplicationVulkanPhysicalDevice, VkDevice InApplicationVulkanDevice)
 	{
-		PRAGMA_DISABLE_DEPRECATION_WARNINGS
 		FrameInfoVulkan.VulkanInstance = InApplicationVulkanInstance;
 		FrameInfoVulkan.VulkanPhysicalDevice = InApplicationVulkanPhysicalDevice;
 		FrameInfoVulkan.VulkanDevice = InApplicationVulkanDevice;
 
 		FrameFormat = EVideoFrameFormat::VULKAN_R8G8B8A8_UNORM;
-		PRAGMA_ENABLE_DEPRECATION_WARNINGS
 
 		CollectAvailableEncoders();
 		return true;
@@ -444,7 +411,6 @@ namespace AVEncoder
 	void FVideoEncoderInputImpl::CollectAvailableEncoders()
 	{
 		AvailableEncoders.Empty();
-		PRAGMA_DISABLE_DEPRECATION_WARNINGS
 		for (const FVideoEncoderInfo& Info : FVideoEncoderFactory::Get().GetAvailable())
 		{
 			if (Info.SupportedInputFormats.Contains(FrameFormat))
@@ -452,21 +418,16 @@ namespace AVEncoder
 				AvailableEncoders.Push(Info);
 			}
 		}
-		PRAGMA_ENABLE_DEPRECATION_WARNINGS
 	}
 
-	PRAGMA_DISABLE_DEPRECATION_WARNINGS
 	const TArray<FVideoEncoderInfo>& FVideoEncoderInputImpl::GetAvailableEncoders()
-	PRAGMA_ENABLE_DEPRECATION_WARNINGS
 	{
 		return AvailableEncoders;
 	}
 
 	// --- encoder input frames -----------------------------------------------------------------------
 
-	PRAGMA_DISABLE_DEPRECATION_WARNINGS
 	bool FVideoEncoderInputImpl::IsUserManagedFrame(const FVideoEncoderInputFrame* InBuffer) const
-	PRAGMA_ENABLE_DEPRECATION_WARNINGS
 	{
 		const FVideoEncoderInputFrameImpl* Frame = static_cast<const FVideoEncoderInputFrameImpl*>(InBuffer);
 		FScopeLock						Guard(&ProtectFrames);
@@ -481,9 +442,7 @@ namespace AVEncoder
 	}
 
 	// create a user managed buffer
-	PRAGMA_DISABLE_DEPRECATION_WARNINGS
 	FVideoEncoderInputFrame* FVideoEncoderInputImpl::CreateBuffer(OnFrameReleasedCallback InOnFrameReleased)
-	PRAGMA_ENABLE_DEPRECATION_WARNINGS
 	{
 		FVideoEncoderInputFrameImpl* Frame = CreateFrame();
 		if (Frame)
@@ -495,9 +454,7 @@ namespace AVEncoder
 	}
 
 	// destroy user managed buffer
-	PRAGMA_DISABLE_DEPRECATION_WARNINGS
 	void FVideoEncoderInputImpl::DestroyBuffer(FVideoEncoderInputFrame* InBuffer)
-	PRAGMA_ENABLE_DEPRECATION_WARNINGS
 	{
 		FVideoEncoderInputFrameImpl* Frame = static_cast<FVideoEncoderInputFrameImpl*>(InBuffer);
 		FScopeLock						Guard(&ProtectFrames);
@@ -518,9 +475,7 @@ namespace AVEncoder
 
 	// --- encoder input frames -----------------------------------------------------------------------
 
-	PRAGMA_DISABLE_DEPRECATION_WARNINGS
 	TSharedPtr<FVideoEncoderInputFrame> FVideoEncoderInputImpl::ObtainInputFrame()
-	PRAGMA_ENABLE_DEPRECATION_WARNINGS
 	{
 		TSharedPtr<FVideoEncoderInputFrameImpl> Frame;
 		FScopeLock						Guard(&ProtectFrames);
@@ -534,15 +489,11 @@ namespace AVEncoder
 		else
 		{
 			Frame = MakeShareable(CreateFrame());
-			PRAGMA_DISABLE_DEPRECATION_WARNINGS
 			UE_LOG(LogVideoEncoder, Verbose, TEXT("Created new frame total frames: %d"), NumBuffers);
-			PRAGMA_ENABLE_DEPRECATION_WARNINGS
 		}
-		
 
 		ActiveFrames.Push(Frame);
 
-		PRAGMA_DISABLE_DEPRECATION_WARNINGS
 		Frame->SetFrameID(NextFrameID++);
 
 		if (NextFrameID == 0)
@@ -551,14 +502,12 @@ namespace AVEncoder
 		}
 
 		Frame->Obtain();
-		PRAGMA_ENABLE_DEPRECATION_WARNINGS
 		return Frame;
 	}
 
 	FVideoEncoderInputFrameImpl* FVideoEncoderInputImpl::CreateFrame()
 	{
 		FVideoEncoderInputFrameImpl* Frame = new FVideoEncoderInputFrameImpl(this);
-		PRAGMA_DISABLE_DEPRECATION_WARNINGS
 		NumBuffers++;
 		switch (FrameFormat)
 		{
@@ -584,13 +533,10 @@ namespace AVEncoder
 			check(false);
 			break;
 		}
-		PRAGMA_ENABLE_DEPRECATION_WARNINGS
 		return Frame;
 	}
 
-	PRAGMA_DISABLE_DEPRECATION_WARNINGS
 	void FVideoEncoderInputImpl::ReleaseInputFrame(FVideoEncoderInputFrame* InFrame)
-	PRAGMA_ENABLE_DEPRECATION_WARNINGS
 	{
 		FVideoEncoderInputFrameImpl* InFrameImpl = static_cast<FVideoEncoderInputFrameImpl*>(InFrame);
 
@@ -618,7 +564,6 @@ namespace AVEncoder
 		if (NumRemoved > 0)
 		{
 			// drop frame if format changed
-			PRAGMA_DISABLE_DEPRECATION_WARNINGS
 			if (InFrame->GetFormat() != FrameFormat)
 			{
 				ProtectFrames.Unlock();
@@ -634,7 +579,6 @@ namespace AVEncoder
 				UE_LOG(LogVideoEncoder, Verbose, TEXT("Deleted buffer (too many) total frames: %d"), NumBuffers);
 				return;
 			}
-			PRAGMA_ENABLE_DEPRECATION_WARNINGS
 
 			AvailableFrames.Enqueue(InFramePtr);
 		}
@@ -644,40 +588,32 @@ namespace AVEncoder
 	{
 		FScopeLock ScopeLock(&ProtectFrames);
 		AvailableFrames.Empty();
-		PRAGMA_DISABLE_DEPRECATION_WARNINGS
 		NumBuffers = 0;
-		PRAGMA_ENABLE_DEPRECATION_WARNINGS
 	}
-	
 
 	void FVideoEncoderInputImpl::SetupFrameYUV420P(FVideoEncoderInputFrameImpl* Frame)
 	{
-		PRAGMA_DISABLE_DEPRECATION_WARNINGS
 		Frame->SetFormat(EVideoFrameFormat::YUV420P);
 		FVideoEncoderInputFrame::FYUV420P& YUV420P = Frame->GetYUV420P();
 		YUV420P.StrideY = FrameInfoYUV420P.StrideY;
 		YUV420P.StrideU = FrameInfoYUV420P.StrideU;
 		YUV420P.StrideV = FrameInfoYUV420P.StrideV;
 		YUV420P.Data[0] = YUV420P.Data[1] = YUV420P.Data[2] = nullptr;
-		PRAGMA_ENABLE_DEPRECATION_WARNINGS
 	}
 
 	void FVideoEncoderInputImpl::SetupFrameD3D11(FVideoEncoderInputFrameImpl* Frame)
 	{
 #if PLATFORM_WINDOWS
-		PRAGMA_DISABLE_DEPRECATION_WARNINGS
 		Frame->SetFormat(FrameFormat);
 
 		FVideoEncoderInputFrame::FD3D11& Data = Frame->GetD3D11();
 		Data.EncoderDevice = FrameInfoD3D.EncoderDeviceD3D11;
-		PRAGMA_ENABLE_DEPRECATION_WARNINGS
 #endif
 	}
 
 	void FVideoEncoderInputImpl::SetupFrameD3D12(FVideoEncoderInputFrameImpl* Frame)
 	{
 #if PLATFORM_WINDOWS
-		PRAGMA_DISABLE_DEPRECATION_WARNINGS
 		Frame->SetFormat(FrameFormat);
 
 		if (FrameFormat == EVideoFrameFormat::D3D11_R8G8B8A8_UNORM)
@@ -690,30 +626,24 @@ namespace AVEncoder
 			FVideoEncoderInputFrame::FD3D12& Data = Frame->GetD3D12();
 			Data.EncoderDevice = FrameInfoD3D.EncoderDeviceD3D12;
 		}
-		PRAGMA_ENABLE_DEPRECATION_WARNINGS
 #endif
 	}
 
 	void FVideoEncoderInputImpl::SetupFrameVulkan(FVideoEncoderInputFrameImpl* Frame)
 	{
 #if PLATFORM_DESKTOP && !PLATFORM_APPLE
-		PRAGMA_DISABLE_DEPRECATION_WARNINGS
 		Frame->SetFormat(FrameFormat);
 
 		FVideoEncoderInputFrame::FVulkan& Data = Frame->GetVulkan();
 		Data.EncoderDevice = FrameInfoVulkan.VulkanDevice;
-		PRAGMA_ENABLE_DEPRECATION_WARNINGS
 #endif
 	}
 
 	void FVideoEncoderInputImpl::SetupFrameCUDA(FVideoEncoderInputFrameImpl* Frame)
 	{
-		PRAGMA_DISABLE_DEPRECATION_WARNINGS
 		Frame->SetFormat(FrameFormat);
-
 		FVideoEncoderInputFrame::FCUDA& Data = Frame->GetCUDA();
 		Data.EncoderDevice = FrameInfoCUDA.EncoderContextCUDA;
-		PRAGMA_ENABLE_DEPRECATION_WARNINGS
 	}
 
 	// ---
@@ -749,29 +679,24 @@ namespace AVEncoder
 		, TimestampUs(0)
 		, TimestampRTP(0)
 		, NumReferences(0)
+		, Format(EVideoFrameFormat::Undefined)
 		, Width(0)
 		, Height(0)
 		, bFreeYUV420PData(false)
-		PRAGMA_DISABLE_DEPRECATION_WARNINGS
-		, Format(EVideoFrameFormat::Undefined)
-		PRAGMA_ENABLE_DEPRECATION_WARNINGS
 	{
 	}
 
-	PRAGMA_DISABLE_DEPRECATION_WARNINGS
 	FVideoEncoderInputFrame::FVideoEncoderInputFrame(const FVideoEncoderInputFrame& CloneFrom)
-	PRAGMA_ENABLE_DEPRECATION_WARNINGS
 		: FrameID(CloneFrom.FrameID)
 		, TimestampUs(CloneFrom.TimestampUs)
 		, TimestampRTP(CloneFrom.TimestampRTP)
 		, NumReferences(0)
+		, Format(CloneFrom.Format)
 		, Width(CloneFrom.Width)
 		, Height(CloneFrom.Height)
 		, bFreeYUV420PData(false)
-		, Format(CloneFrom.Format)
 	{
 #if PLATFORM_WINDOWS
-		PRAGMA_DISABLE_DEPRECATION_WARNINGS
 		D3D11.EncoderDevice = CloneFrom.D3D11.EncoderDevice;
 		D3D11.Texture = CloneFrom.D3D11.Texture;
 		if ((D3D11.EncoderTexture = CloneFrom.D3D11.EncoderTexture) != nullptr)
@@ -785,49 +710,37 @@ namespace AVEncoder
 		{
 			D3D12.EncoderTexture->AddRef();
 		}
-		PRAGMA_ENABLE_DEPRECATION_WARNINGS
 #endif
 
-		PRAGMA_DISABLE_DEPRECATION_WARNINGS
 		CUDA.EncoderDevice = CloneFrom.CUDA.EncoderDevice;
 		CUDA.EncoderTexture = CloneFrom.CUDA.EncoderTexture;
-		PRAGMA_ENABLE_DEPRECATION_WARNINGS
 	}
 
 	FVideoEncoderInputFrame::~FVideoEncoderInputFrame()
 	{
 		if (bFreeYUV420PData)
 		{
-			PRAGMA_DISABLE_DEPRECATION_WARNINGS
 			delete[] YUV420P.Data[0];
 			delete[] YUV420P.Data[1];
 			delete[] YUV420P.Data[2];
 			YUV420P.Data[0] = YUV420P.Data[1] = YUV420P.Data[2] = nullptr;
-			PRAGMA_ENABLE_DEPRECATION_WARNINGS
 			bFreeYUV420PData = false;
 		}
 
 #if PLATFORM_WINDOWS
-		PRAGMA_DISABLE_DEPRECATION_WARNINGS
 		if (D3D11.EncoderTexture)
-		PRAGMA_ENABLE_DEPRECATION_WARNINGS
 		{
 			// check to make sure this frame holds the last reference
-			PRAGMA_DISABLE_DEPRECATION_WARNINGS
 			auto	NumRef = D3D11.EncoderTexture->AddRef();
-			PRAGMA_ENABLE_DEPRECATION_WARNINGS
 			if (NumRef > 2)
 			{
 				UE_LOG(LogVideoEncoder, Warning, TEXT("VideoEncoderFame - D3D11 input texture still holds %d references."), NumRef);
 			}
 			// Need to call release twice as we have added an extra reference just above (only way to count how many references we have)
-			PRAGMA_DISABLE_DEPRECATION_WARNINGS
 			D3D11.EncoderTexture->Release();
 			D3D11.EncoderTexture->Release();
 			D3D11.EncoderTexture = nullptr;
-			PRAGMA_ENABLE_DEPRECATION_WARNINGS
 		}
-		PRAGMA_DISABLE_DEPRECATION_WARNINGS
 		if (D3D11.SharedHandle)
 		{
 			CloseHandle(D3D11.SharedHandle);
@@ -855,28 +768,22 @@ namespace AVEncoder
 			OnReleaseD3D12Texture(D3D12.Texture);
 			D3D12.Texture = nullptr;
 		}
-		PRAGMA_ENABLE_DEPRECATION_WARNINGS
 
 		// D3D12 specific handle atm
-		PRAGMA_DISABLE_DEPRECATION_WARNINGS
 		if (CUDA.SharedHandle)
 		{
 			CloseHandle(CUDA.SharedHandle);
 			CUDA.SharedHandle = nullptr;
 		}
-		PRAGMA_ENABLE_DEPRECATION_WARNINGS
 #endif 
 
-		PRAGMA_DISABLE_DEPRECATION_WARNINGS
 		if (CUDA.EncoderTexture)
 		{
 			OnReleaseCUDATexture(CUDA.EncoderTexture);
 			CUDA.EncoderTexture = nullptr;
 		}
-		PRAGMA_ENABLE_DEPRECATION_WARNINGS
 
 #if PLATFORM_DESKTOP && !PLATFORM_APPLE
-		PRAGMA_DISABLE_DEPRECATION_WARNINGS
 		if (Vulkan.EncoderTexture != VK_NULL_HANDLE)
 		{
 			OnReleaseVulkanTexture(Vulkan.EncoderTexture);
@@ -886,7 +793,6 @@ namespace AVEncoder
 		{
 			OnReleaseVulkanSurface(Vulkan.EncoderSurface);
 		}
-		PRAGMA_ENABLE_DEPRECATION_WARNINGS
 #endif
 	}
 
@@ -894,59 +800,45 @@ namespace AVEncoder
 	{
 		if (!bFreeYUV420PData)
 		{
-			PRAGMA_DISABLE_DEPRECATION_WARNINGS
 			YUV420P.StrideY = Width;
 			YUV420P.StrideU = (Width + 1) / 2;
 			YUV420P.StrideV = (Width + 1) / 2;;
 			YUV420P.Data[0] = new uint8[Height * YUV420P.StrideY];
 			YUV420P.Data[1] = new uint8[(Height + 1) / 2 * YUV420P.StrideU];
 			YUV420P.Data[2] = new uint8[(Height + 1) / 2 * YUV420P.StrideV];
-			PRAGMA_ENABLE_DEPRECATION_WARNINGS
 			bFreeYUV420PData = true;
 		}
 	}
 
 	void FVideoEncoderInputFrame::SetYUV420P(const uint8* InDataY, const uint8* InDataU, const uint8* InDataV, uint32 InStrideY, uint32 InStrideU, uint32 InStrideV)
 	{
-		PRAGMA_DISABLE_DEPRECATION_WARNINGS
 		if (Format == EVideoFrameFormat::YUV420P)
-		PRAGMA_ENABLE_DEPRECATION_WARNINGS
 		{
 			if (bFreeYUV420PData)
 			{
-				PRAGMA_DISABLE_DEPRECATION_WARNINGS
 				delete[] YUV420P.Data[0];
 				delete[] YUV420P.Data[1];
 				delete[] YUV420P.Data[2];
-				PRAGMA_ENABLE_DEPRECATION_WARNINGS
 				bFreeYUV420PData = false;
 			}
-			PRAGMA_DISABLE_DEPRECATION_WARNINGS
 			YUV420P.Data[0] = InDataY;
 			YUV420P.Data[1] = InDataU;
 			YUV420P.Data[2] = InDataV;
 			YUV420P.StrideY = InStrideY;
 			YUV420P.StrideU = InStrideU;
 			YUV420P.StrideV = InStrideV;
-			PRAGMA_ENABLE_DEPRECATION_WARNINGS
 		}
 	}
 
 	static FThreadSafeCounter	_VideoEncoderInputFrameCnt{ 0 };
 
 #if PLATFORM_WINDOWS
-	PRAGMA_DISABLE_DEPRECATION_WARNINGS
 	void FVideoEncoderInputFrame::SetTexture(ID3D11Texture2D* InTexture, FReleaseD3D11TextureCallback InOnReleaseTexture)
-	PRAGMA_ENABLE_DEPRECATION_WARNINGS
 	{
-		PRAGMA_DISABLE_DEPRECATION_WARNINGS
 		if (Format == EVideoFrameFormat::D3D11_R8G8B8A8_UNORM)
-		PRAGMA_ENABLE_DEPRECATION_WARNINGS
 		{
-			PRAGMA_DISABLE_DEPRECATION_WARNINGS
 			check(D3D11.Texture == nullptr);
 			if (!D3D11.Texture)
-			PRAGMA_ENABLE_DEPRECATION_WARNINGS
 			{
 				TRefCountPtr<IDXGIResource> DXGIResource;
 				HANDLE						SharedHandle;
@@ -966,38 +858,28 @@ namespace AVEncoder
 				{
 					UE_LOG(LogVideoEncoder, Error, TEXT("IDXGIResource::GetSharedHandle() failed to return a shared texture resource no created as shared? (D3D11_RESOURCE_MISC_SHARED)."));
 				}
-				PRAGMA_DISABLE_DEPRECATION_WARNINGS
 				else if ((Result = D3D11.EncoderDevice->OpenSharedResource(SharedHandle, __uuidof(ID3D11Texture2D), (LPVOID*)&D3D11.EncoderTexture)) != S_OK)
-				PRAGMA_ENABLE_DEPRECATION_WARNINGS
 				{
 					UE_LOG(LogVideoEncoder, Error, TEXT("ID3D11Device::OpenSharedResource() failed 0x%X - %s."), Result, *GetComErrorDescription(Result));
 				}
 				else
 				{
-					PRAGMA_DISABLE_DEPRECATION_WARNINGS
 					DEBUG_SET_D3D11_OBJECT_NAME(D3D11.EncoderTexture, "FVideoEncoderInputFrame::SetTexture()");
 					D3D11.Texture = InTexture;
-					PRAGMA_ENABLE_DEPRECATION_WARNINGS
 					OnReleaseD3D11Texture = InOnReleaseTexture;
 				}
 			}
 		}
 	}
 
-	PRAGMA_DISABLE_DEPRECATION_WARNINGS
 	void FVideoEncoderInputFrame::SetTexture(ID3D12Resource* InTexture, FReleaseD3D12TextureCallback InOnReleaseTexture)
-	PRAGMA_ENABLE_DEPRECATION_WARNINGS
 	{
-		PRAGMA_DISABLE_DEPRECATION_WARNINGS
 		if (Format == EVideoFrameFormat::D3D11_R8G8B8A8_UNORM)
-		PRAGMA_ENABLE_DEPRECATION_WARNINGS
 		{
-			PRAGMA_DISABLE_DEPRECATION_WARNINGS
 			check(D3D12.Texture == nullptr)
-			check(D3D11.EncoderTexture == nullptr)
-			PRAGMA_ENABLE_DEPRECATION_WARNINGS
+				check(D3D11.EncoderTexture == nullptr)
 
-			TRefCountPtr<ID3D12Device>	OwnerDevice;
+				TRefCountPtr<ID3D12Device>	OwnerDevice;
 			HRESULT						Result;
 			if ((Result = InTexture->GetDevice(IID_PPV_ARGS(OwnerDevice.GetInitReference()))) != S_OK)
 			{
@@ -1006,65 +888,48 @@ namespace AVEncoder
 			//
 			// NOTE: ID3D12Device::CreateSharedHandle gives as an NT Handle, and so we need to call CloseHandle on it
 			//
-			PRAGMA_DISABLE_DEPRECATION_WARNINGS
 			else if ((Result = OwnerDevice->CreateSharedHandle(InTexture, NULL, GENERIC_ALL, *FString::Printf(TEXT("%s_FVideoEncoderInputFrame_%d"), *GetGUID(), _VideoEncoderInputFrameCnt.Increment()), &D3D11.SharedHandle)) != S_OK)
-			PRAGMA_ENABLE_DEPRECATION_WARNINGS
 			{
 				UE_LOG(LogVideoEncoder, Error, TEXT("ID3D12Device::CreateSharedHandle() failed 0x%X - %s."), Result, *GetComErrorDescription(Result));
 			}
-			PRAGMA_DISABLE_DEPRECATION_WARNINGS
 			else if (D3D11.SharedHandle == nullptr)
-			PRAGMA_ENABLE_DEPRECATION_WARNINGS
 			{
 				UE_LOG(LogVideoEncoder, Error, TEXT("ID3D12Device::CreateSharedHandle() failed to return a shared texture resource no created as shared? (D3D12_RESOURCE_FLAG_ALLOW_SIMULTANEOUS_ACCESS)."));
 			}
 			else
 			{
 				TRefCountPtr <ID3D11Device1> Device1;
-				PRAGMA_DISABLE_DEPRECATION_WARNINGS
 				if ((Result = D3D11.EncoderDevice->QueryInterface(IID_PPV_ARGS(Device1.GetInitReference()))) != S_OK)
-				PRAGMA_ENABLE_DEPRECATION_WARNINGS
 				{
 					UE_LOG(LogVideoEncoder, Error, TEXT("ID3D11Device::QueryInterface() failed 0x%X - %s."), Result, *GetComErrorDescription(Result));
 				}
-				PRAGMA_DISABLE_DEPRECATION_WARNINGS
 				else if ((Result = Device1->OpenSharedResource1(D3D11.SharedHandle, IID_PPV_ARGS(&D3D11.EncoderTexture))) != S_OK)
-				PRAGMA_ENABLE_DEPRECATION_WARNINGS
 				{
 					UE_LOG(LogVideoEncoder, Error, TEXT("ID3D11Device::OpenSharedResource1() failed 0x%X - %s."), Result, *GetComErrorDescription(Result));
 				}
 				else
 				{
-					PRAGMA_DISABLE_DEPRECATION_WARNINGS
 					DEBUG_SET_D3D11_OBJECT_NAME(D3D11.EncoderTexture, "FVideoEncoderInputFrame::SetTexture()");
 					D3D12.Texture = InTexture;
-					PRAGMA_ENABLE_DEPRECATION_WARNINGS
 					OnReleaseD3D12Texture = InOnReleaseTexture;
 				}
 			}
 		}
 		else
 		{
-			PRAGMA_DISABLE_DEPRECATION_WARNINGS
 			check(D3D12.Texture == nullptr);
 			check(D3D12.EncoderTexture == nullptr);
 			D3D12.Texture = InTexture;
 			D3D12.EncoderTexture = InTexture;
-			PRAGMA_ENABLE_DEPRECATION_WARNINGS
 			OnReleaseD3D12Texture = InOnReleaseTexture;
 		}
 	}
 #endif
 
-	PRAGMA_DISABLE_DEPRECATION_WARNINGS
 	void FVideoEncoderInputFrame::SetTexture(CUarray InTexture, EUnderlyingRHI UnderlyingRHI, void* SharedHandle, FReleaseCUDATextureCallback InOnReleaseTexture)
-	PRAGMA_ENABLE_DEPRECATION_WARNINGS
 	{
-		PRAGMA_DISABLE_DEPRECATION_WARNINGS
 		if (Format == EVideoFrameFormat::CUDA_R8G8B8A8_UNORM)
-		PRAGMA_ENABLE_DEPRECATION_WARNINGS
 		{
-			PRAGMA_DISABLE_DEPRECATION_WARNINGS
 			CUDA.UnderlyingRHI = UnderlyingRHI;
 			CUDA.SharedHandle = SharedHandle;
 			CUDA.EncoderTexture = InTexture;
@@ -1073,40 +938,28 @@ namespace AVEncoder
 			{
 				UE_LOG(LogVideoEncoder, Warning, TEXT("SetTexture | Cuda device pointer is null"));
 			}
-			PRAGMA_ENABLE_DEPRECATION_WARNINGS
 		}
 	}
 
 
 #if PLATFORM_DESKTOP && !PLATFORM_APPLE
-	PRAGMA_DISABLE_DEPRECATION_WARNINGS
 	void FVideoEncoderInputFrame::SetTexture(VkImage InTexture, FReleaseVulkanTextureCallback InOnReleaseTexture)
-	PRAGMA_ENABLE_DEPRECATION_WARNINGS
 	{
-		PRAGMA_DISABLE_DEPRECATION_WARNINGS
 		if (Format == EVideoFrameFormat::VULKAN_R8G8B8A8_UNORM)
-		PRAGMA_ENABLE_DEPRECATION_WARNINGS
 		{
-			PRAGMA_DISABLE_DEPRECATION_WARNINGS
 			Vulkan.EncoderTexture = InTexture;
 			OnReleaseVulkanTexture = InOnReleaseTexture;
 			if (!Vulkan.EncoderTexture)
 			{
 				UE_LOG(LogVideoEncoder, Warning, TEXT("SetTexture | Vulkan VkImage is null"));
 			}
-			PRAGMA_ENABLE_DEPRECATION_WARNINGS
 		}
 	}
 
-	PRAGMA_DISABLE_DEPRECATION_WARNINGS
 	void FVideoEncoderInputFrame::SetTexture(VkImage InTexture, VkDeviceMemory InTextureDeviceMemory, uint64 InTextureMemorySize, FReleaseVulkanTextureCallback InOnReleaseTexture)
-	PRAGMA_ENABLE_DEPRECATION_WARNINGS
 	{
-		PRAGMA_DISABLE_DEPRECATION_WARNINGS
 		if (Format == EVideoFrameFormat::VULKAN_R8G8B8A8_UNORM)
-		PRAGMA_ENABLE_DEPRECATION_WARNINGS
 		{
-			PRAGMA_DISABLE_DEPRECATION_WARNINGS
 			Vulkan.EncoderTexture = InTexture;
 			Vulkan.EncoderDeviceMemory = InTextureDeviceMemory;
 			Vulkan.EncoderMemorySize = InTextureMemorySize;
@@ -1116,21 +969,17 @@ namespace AVEncoder
 			{
 				UE_LOG(LogVideoEncoder, Warning, TEXT("SetTexture | Vulkan VkImage or VkTextureDeviceMemory is null"));
 			}
-			PRAGMA_ENABLE_DEPRECATION_WARNINGS
 		}
 	}
 #endif
 
 	// *** FVideoEncoderInputFrameImpl ****************************************************************
 
-	PRAGMA_DISABLE_DEPRECATION_WARNINGS
 	FVideoEncoderInputFrameImpl::FVideoEncoderInputFrameImpl(FVideoEncoderInputImpl* InInput)
 		: Input(InInput)
 	{
 	}
-	PRAGMA_ENABLE_DEPRECATION_WARNINGS
 
-	PRAGMA_DISABLE_DEPRECATION_WARNINGS
 	FVideoEncoderInputFrameImpl::FVideoEncoderInputFrameImpl(const FVideoEncoderInputFrameImpl& InCloneFrom, FCloneDestroyedCallback InCloneDestroyedCallback)
 		: FVideoEncoderInputFrame(InCloneFrom)
 		, Input(InCloneFrom.Input)
@@ -1138,9 +987,7 @@ namespace AVEncoder
 		, OnCloneDestroyed(MoveTemp(InCloneDestroyedCallback))
 	{
 	}
-	PRAGMA_ENABLE_DEPRECATION_WARNINGS
 
-	PRAGMA_DISABLE_DEPRECATION_WARNINGS
 	FVideoEncoderInputFrameImpl::~FVideoEncoderInputFrameImpl()
 	{
 		if (ClonedReference)
@@ -1151,7 +998,6 @@ namespace AVEncoder
 		{
 		}
 	}
-	PRAGMA_ENABLE_DEPRECATION_WARNINGS
 
 	void FVideoEncoderInputFrameImpl::Release() const
 	{
@@ -1160,9 +1006,7 @@ namespace AVEncoder
 		{
 			Input->ReleaseInputFrame(const_cast<FVideoEncoderInputFrameImpl*>(this));
 		}
-		PRAGMA_DISABLE_DEPRECATION_WARNINGS
 		else if (NumReferences.Decrement() == 0)
-		PRAGMA_ENABLE_DEPRECATION_WARNINGS
 		{
 			if (ClonedReference)
 			{
@@ -1177,11 +1021,12 @@ namespace AVEncoder
 	}
 
 	// Clone frame - this will create a copy that references the original until destroyed
-	PRAGMA_DISABLE_DEPRECATION_WARNINGS
 	const FVideoEncoderInputFrame* FVideoEncoderInputFrameImpl::Clone(FCloneDestroyedCallback InCloneDestroyedCallback) const
-	PRAGMA_ENABLE_DEPRECATION_WARNINGS
 	{
 		FVideoEncoderInputFrameImpl* ClonedFrame = new FVideoEncoderInputFrameImpl(*this, MoveTemp(InCloneDestroyedCallback));
 		return ClonedFrame;
 	}
+
+
+
 } /* namespace AVEncoder */
diff --git a/Engine/Source/Runtime/AVEncoder/Private/VideoEncoderInputImpl.h b/Engine/Source/Runtime/AVEncoder/Private/VideoEncoderInputImpl.h
--- a/Engine/Source/Runtime/AVEncoder/Private/VideoEncoderInputImpl.h
+++ b/Engine/Source/Runtime/AVEncoder/Private/VideoEncoderInputImpl.h
@@ -17,7 +17,6 @@ namespace AVEncoder
 {
 class FVideoEncoderInputFrameImpl;
 
-PRAGMA_DISABLE_DEPRECATION_WARNINGS
 class FVideoEncoderInputImpl : public FVideoEncoderInput
 {
 public:
@@ -166,5 +165,5 @@ private:
 	const FVideoEncoderInputFrame*			ClonedReference = nullptr;
 	const FCloneDestroyedCallback			OnCloneDestroyed;
 };
-PRAGMA_ENABLE_DEPRECATION_WARNINGS
+
 } /* namespace AVEncoder */
diff --git a/Engine/Source/Runtime/AVEncoder/Private/VideoEncoderModule.cpp b/Engine/Source/Runtime/AVEncoder/Private/VideoEncoderModule.cpp
--- a/Engine/Source/Runtime/AVEncoder/Private/VideoEncoderModule.cpp
+++ b/Engine/Source/Runtime/AVEncoder/Private/VideoEncoderModule.cpp
@@ -11,9 +11,7 @@ public:
 
 	void ShutdownModule()
 	{
-		PRAGMA_DISABLE_DEPRECATION_WARNINGS
 		AVEncoder::FVideoEncoderFactory::Shutdown();
-		PRAGMA_ENABLE_DEPRECATION_WARNINGS
 	}
 };
 
diff --git a/Engine/Source/Runtime/AVEncoder/Public/AudioEncoder.h b/Engine/Source/Runtime/AVEncoder/Public/AudioEncoder.h
--- a/Engine/Source/Runtime/AVEncoder/Public/AudioEncoder.h
+++ b/Engine/Source/Runtime/AVEncoder/Public/AudioEncoder.h
@@ -12,55 +12,44 @@
 
 namespace AVEncoder
 {
-	struct UE_DEPRECATED(5.4, "AVEncoder has been deprecated. Please use the AVCodecs plugin family instead.") FAudioFrame
+	struct FAudioFrame
 	{
 		FTimespan Timestamp;
 		FTimespan Duration;
 		Audio::TSampleBuffer<float> Data;
 	};
 
-	class UE_DEPRECATED(5.4, "AVEncoder has been deprecated. Please use the AVCodecs plugin family instead.") IAudioEncoderListener
+	class IAudioEncoderListener
 	{
 	public:
-		PRAGMA_DISABLE_DEPRECATION_WARNINGS
 		virtual void OnEncodedAudioFrame(const FMediaPacket& Packet) = 0;
-		PRAGMA_ENABLE_DEPRECATION_WARNINGS
 	};
 
-	class UE_DEPRECATED(5.4, "AVEncoder has been deprecated. Please use the AVCodecs plugin family instead.") FAudioEncoder
+	class FAudioEncoder
 	{
 	public:
 		virtual ~FAudioEncoder() {}
 		virtual const TCHAR* GetName() const = 0;
 		virtual const TCHAR* GetType() const = 0;
-		PRAGMA_DISABLE_DEPRECATION_WARNINGS
 		virtual bool Initialize(const FAudioConfig& Config) = 0;
-		PRAGMA_ENABLE_DEPRECATION_WARNINGS
 
 		/**
 		* Shutdown MUST be called before destruction
 		*/
 		virtual void Shutdown() = 0;
 
-		PRAGMA_DISABLE_DEPRECATION_WARNINGS
 		virtual void Encode(const FAudioFrame& Frame) = 0;
 		virtual FAudioConfig GetConfig() const = 0;
 
 		AVENCODER_API virtual void RegisterListener(IAudioEncoderListener& Listener);
 		AVENCODER_API virtual void UnregisterListener(IAudioEncoderListener& Listener);
-		PRAGMA_ENABLE_DEPRECATION_WARNINGS
 
 	protected:
-		PRAGMA_DISABLE_DEPRECATION_WARNINGS
 		AVENCODER_API void OnEncodedAudioFrame(const FMediaPacket& Packet);
-		PRAGMA_ENABLE_DEPRECATION_WARNINGS
 	private:
 
 		FCriticalSection ListenersMutex;
-
-		PRAGMA_DISABLE_DEPRECATION_WARNINGS
 		TArray<IAudioEncoderListener*> Listeners;
-		PRAGMA_ENABLE_DEPRECATION_WARNINGS
 
 	};
 }
diff --git a/Engine/Source/Runtime/AVEncoder/Public/AudioEncoderFactory.h b/Engine/Source/Runtime/AVEncoder/Public/AudioEncoderFactory.h
--- a/Engine/Source/Runtime/AVEncoder/Public/AudioEncoderFactory.h
+++ b/Engine/Source/Runtime/AVEncoder/Public/AudioEncoderFactory.h
@@ -11,15 +11,13 @@ namespace AVEncoder
 {
 	class FAudioEncoder;
 
-	class UE_DEPRECATED(5.4, "AVEncoder has been deprecated. Please use the AVCodecs plugin family instead.") FAudioEncoderFactory
+	class FAudioEncoderFactory
 	{
 	public:
 		virtual ~FAudioEncoderFactory() {}
 		virtual const TCHAR* GetName() const = 0;
 		virtual TArray<FString> GetSupportedCodecs() const = 0;
-		PRAGMA_DISABLE_DEPRECATION_WARNINGS
 		virtual TUniquePtr<FAudioEncoder> CreateEncoder(const FString& Codec) = 0;
-		PRAGMA_ENABLE_DEPRECATION_WARNINGS
 
 		static AVENCODER_API void RegisterFactory(FAudioEncoderFactory& Factory);
 		static AVENCODER_API void UnregisterFactory(FAudioEncoderFactory& Factory);
diff --git a/Engine/Source/Runtime/AVEncoder/Public/AudioTypes.h b/Engine/Source/Runtime/AVEncoder/Public/AudioTypes.h
--- a/Engine/Source/Runtime/AVEncoder/Public/AudioTypes.h
+++ b/Engine/Source/Runtime/AVEncoder/Public/AudioTypes.h
@@ -17,7 +17,7 @@ namespace LibRtc
 	using FAudioChannelSample = int16_t;
 
 	/** Describes the data format of the audio stream */
-	struct UE_DEPRECATED(5.4, "AVEncoder has been deprecated. Please use the AVCodecs plugin family instead.") FAudioStreamFormat
+	struct FAudioStreamFormat
 	{
 		/** The number of samples per seconds */
 		uint32_t SampleRate;
@@ -42,7 +42,7 @@ namespace LibRtc
 	 * One second - one frame * SampleRate
 	 * The number of samples per seconds - one sample byte * NumberBytesPerSample * NumberChannels * SampleRate
 	 */
-	struct UE_DEPRECATED(5.4, "AVEncoder has been deprecated. Please use the AVCodecs plugin family instead.") FAudioSampleBuffer
+	struct FAudioSampleBuffer
 	{
 		/** The sequence of audio samples */
 		TArrayView<const FAudioChannelSample> AudioSamples;
@@ -50,7 +50,7 @@ namespace LibRtc
 		FAudioStreamFormat AudioStreamFormat;
 	};
 
-	struct UE_DEPRECATED(5.4, "AVEncoder has been deprecated. Please use the AVCodecs plugin family instead.") FAudioCodecConfig
+	struct FAudioCodecConfig
 	{
 		int32_t SampleRate;
 		int32_t NumChannels;
@@ -78,7 +78,7 @@ namespace LibRtc
 	 */
 	using OnAudioDecodeCallback = TFunction<void(const FAudioByteBuffer& AudioBytes, const FAudioSampleBuffer&)>;
 
-	struct UE_DEPRECATED(5.4, "AVEncoder has been deprecated. Please use the AVCodecs plugin family instead.") FAudioCodec
+	struct FAudioCodec
 	{
 		FAudioCodecConfig AudioCodecConfig;
 		OnAudioEncodeCallback OnAudioEncodeCallback;
@@ -86,19 +86,19 @@ namespace LibRtc
 	};
 
 	/** Information which is used to initialize the incoming audio channel */
-	struct UE_DEPRECATED(5.4, "AVEncoder has been deprecated. Please use the AVCodecs plugin family instead.") FIncomingAudioChannelInfo
+	struct FIncomingAudioChannelInfo
 	{
 		// nothing
 	};
 
 	/** Information which is used to initialize the outgoing audio channel */
-	struct UE_DEPRECATED(5.4, "AVEncoder has been deprecated. Please use the AVCodecs plugin family instead.") FOutgoingAudioChannelInfo
+	struct FOutgoingAudioChannelInfo
 	{
 		/** Initial muted state of the audio channel */
 		bool bIsMuted;
 	};
 
-	enum class UE_DEPRECATED(5.4, "AVEncoder has been deprecated. Please use the AVCodecs plugin family instead.") EAudioMuteReason
+	enum class EAudioMuteReason
 	{
 		Unsupported,
 		Manual,
@@ -106,7 +106,7 @@ namespace LibRtc
 		AdminDisabled
 	};
 
-	struct UE_DEPRECATED(5.4, "AVEncoder has been deprecated. Please use the AVCodecs plugin family instead.") FAudioMuteInfo
+	struct FAudioMuteInfo
 	{
 		TOptional<EAudioMuteReason> Reason;
 	};
diff --git a/Engine/Source/Runtime/AVEncoder/Public/CodecPacket.h b/Engine/Source/Runtime/AVEncoder/Public/CodecPacket.h
--- a/Engine/Source/Runtime/AVEncoder/Public/CodecPacket.h
+++ b/Engine/Source/Runtime/AVEncoder/Public/CodecPacket.h
@@ -9,7 +9,7 @@
 
 namespace AVEncoder
 {
-	class UE_DEPRECATED(5.4, "AVEncoder has been deprecated. Please use the AVCodecs plugin family instead.") FCodecPacket
+	class FCodecPacket
 	{
 	public:
 		virtual ~FCodecPacket() = default;
@@ -19,7 +19,7 @@ namespace AVEncoder
 		/**
 		 * Encoding/Decoding latency
 		 */
-		struct UE_DEPRECATED(5.4, "AVEncoder has been deprecated. Please use the AVCodecs plugin family instead.") FTimings
+		struct FTimings
 		{
 			FTimespan StartTs;
 			FTimespan FinishTs;
diff --git a/Engine/Source/Runtime/AVEncoder/Public/Configuration.h b/Engine/Source/Runtime/AVEncoder/Public/Configuration.h
--- a/Engine/Source/Runtime/AVEncoder/Public/Configuration.h
+++ b/Engine/Source/Runtime/AVEncoder/Public/Configuration.h
@@ -12,9 +12,9 @@
 
 namespace LibRtc
 {
-    struct UE_DEPRECATED(5.4, "AVEncoder has been deprecated. Please use the AVCodecs plugin family instead.") FConfiguration
+    struct FConfiguration
     {
-        struct UE_DEPRECATED(5.4, "AVEncoder has been deprecated. Please use the AVCodecs plugin family instead.") FData
+        struct FData
         {
             static constexpr const TCHAR* kDefaultStunServerUrl = TEXT("stun:stun.l.google.com:19302");
 
diff --git a/Engine/Source/Runtime/AVEncoder/Public/IMFSampleWrapper.h b/Engine/Source/Runtime/AVEncoder/Public/IMFSampleWrapper.h
--- a/Engine/Source/Runtime/AVEncoder/Public/IMFSampleWrapper.h
+++ b/Engine/Source/Runtime/AVEncoder/Public/IMFSampleWrapper.h
@@ -15,16 +15,15 @@ namespace AVEncoder
 //
 // Wrapper for IMFSample, to make it easier to report errors
 //
-class UE_DEPRECATED(5.4, "AVEncoder has been deprecated. Please use the AVCodecs plugin family instead.") FIMFSampleWrapper
+class FIMFSampleWrapper
 {
 public:
-	PRAGMA_DISABLE_DEPRECATION_WARNINGS
+
 	FIMFSampleWrapper(EPacketType InMediaType = EPacketType::Invalid, IMFSample* InSample = nullptr)
 		: MediaType(InMediaType)
 		, Sample(InSample)
 	{
 	}
-	PRAGMA_ENABLE_DEPRECATION_WARNINGS
 
 	const IMFSample* GetSample() const
 	{
@@ -96,9 +95,7 @@ public:
 	AVENCODER_API FIMFSampleWrapper Clone() const;
 
 private:
-	PRAGMA_DISABLE_DEPRECATION_WARNINGS
 	EPacketType MediaType;
-	PRAGMA_ENABLE_DEPRECATION_WARNINGS
 	TRefCountPtr<IMFSample> Sample;
 };
 
diff --git a/Engine/Source/Runtime/AVEncoder/Public/MediaPacket.h b/Engine/Source/Runtime/AVEncoder/Public/MediaPacket.h
--- a/Engine/Source/Runtime/AVEncoder/Public/MediaPacket.h
+++ b/Engine/Source/Runtime/AVEncoder/Public/MediaPacket.h
@@ -5,7 +5,7 @@
 
 namespace AVEncoder
 {
-	struct UE_DEPRECATED(5.4, "AVEncoder has been deprecated. Please use the AVCodecs plugin family instead.") FVideoConfig
+	struct FVideoConfig
 	{
 		FString Codec;
 		uint32 Width;
@@ -14,7 +14,7 @@ namespace AVEncoder
 		uint32 Bitrate;
 	};
 
-	struct UE_DEPRECATED(5.4, "AVEncoder has been deprecated. Please use the AVCodecs plugin family instead.") FAudioConfig
+	struct FAudioConfig
 	{
 		FString Codec;
 		uint32 Samplerate;
@@ -22,18 +22,16 @@ namespace AVEncoder
 		uint32 Bitrate;
 	};
 
-	enum class UE_DEPRECATED(5.4, "AVEncoder has been deprecated. Please use the AVCodecs plugin family instead.") EPacketType
+	enum class EPacketType
 	{
 		Audio,
 		Video,
 		Invalid
 	};
 
-	struct UE_DEPRECATED(5.4, "AVEncoder has been deprecated. Please use the AVCodecs plugin family instead.") FMediaPacket
+	struct FMediaPacket
 	{
-		PRAGMA_DISABLE_DEPRECATION_WARNINGS
 		EPacketType Type;
-		PRAGMA_ENABLE_DEPRECATION_WARNINGS
 		FTimespan Timestamp;
 		FTimespan Duration;
 		TArray<uint8> Data;
@@ -54,19 +52,13 @@ namespace AVEncoder
 			} Audio;
 		};
 
-		PRAGMA_DISABLE_DEPRECATION_WARNINGS
 		explicit FMediaPacket(EPacketType TypeIn)
-		PRAGMA_ENABLE_DEPRECATION_WARNINGS
 		{
 			Type = TypeIn;
-			PRAGMA_DISABLE_DEPRECATION_WARNINGS
 			if (Type == EPacketType::Audio)
-			PRAGMA_ENABLE_DEPRECATION_WARNINGS
 			{
 			}
-			PRAGMA_DISABLE_DEPRECATION_WARNINGS
 			else if (Type == EPacketType::Video)
-			PRAGMA_ENABLE_DEPRECATION_WARNINGS
 			{
 				FMemory::Memzero(Video);
 			}
@@ -74,9 +66,7 @@ namespace AVEncoder
 
 		bool IsVideoKeyFrame() const
 		{
-			PRAGMA_DISABLE_DEPRECATION_WARNINGS
 			return (Type == EPacketType::Video && Video.bKeyFrame) ? true : false;
-			PRAGMA_ENABLE_DEPRECATION_WARNINGS
 		}
 	};
 }
diff --git a/Engine/Source/Runtime/AVEncoder/Public/VideoCommon.h b/Engine/Source/Runtime/AVEncoder/Public/VideoCommon.h
--- a/Engine/Source/Runtime/AVEncoder/Public/VideoCommon.h
+++ b/Engine/Source/Runtime/AVEncoder/Public/VideoCommon.h
@@ -47,7 +47,7 @@ namespace AVEncoder
 {
 	const int64 TimeStampNone = 0x7fffffffll;
 
-	enum class UE_DEPRECATED(5.4, "AVEncoder has been deprecated. Please use the AVCodecs plugin family instead.") EVideoFrameFormat
+	enum class EVideoFrameFormat
 	{
 		Undefined,				// (not-yet) defined format
 		YUV420P,				// Planar YUV420 format in CPU memory
@@ -57,7 +57,7 @@ namespace AVEncoder
 		VULKAN_R8G8B8A8_UNORM,
 	};
 
-	enum class UE_DEPRECATED(5.4, "AVEncoder has been deprecated. Please use the AVCodecs plugin family instead.") EH264Profile
+	enum class EH264Profile
 	{
 		UNKNOWN,
 		CONSTRAINED_BASELINE,
@@ -67,7 +67,6 @@ namespace AVEncoder
 		HIGH,
 	};
 
-	PRAGMA_DISABLE_DEPRECATION_WARNINGS
 	inline FString ToString(EVideoFrameFormat Format)
 	{
 		switch (Format)
@@ -87,9 +86,8 @@ namespace AVEncoder
 			return FString("EVideoFrameFormat::Undefined");
 		}
 	}
-	PRAGMA_ENABLE_DEPRECATION_WARNINGS
 
-	enum class UE_DEPRECATED(5.4, "AVEncoder has been deprecated. Please use the AVCodecs plugin family instead.") ECodecType
+	enum class ECodecType
 	{
 		Undefined,
 		H264,
@@ -104,16 +102,13 @@ namespace AVEncoder
 	const uint32 H264Profile_ConstrainedHigh = 1 << 3;
 	const uint32 H264Profile_High = 1 << 4;
 
-	struct UE_DEPRECATED(5.4, "AVEncoder has been deprecated. Please use the AVCodecs plugin family instead.") FVideoEncoderInfo
+	struct FVideoEncoderInfo
 	{
 		uint32						ID = 0;
+		ECodecType					CodecType = ECodecType::Undefined;
 		uint32						MaxWidth = 0;
 		uint32						MaxHeight = 0;
-		PRAGMA_DISABLE_DEPRECATION_WARNINGS
-		ECodecType					CodecType = ECodecType::Undefined;
 		TArray<EVideoFrameFormat>	SupportedInputFormats;
-		PRAGMA_ENABLE_DEPRECATION_WARNINGS
-
 		struct
 		{
 			uint32					SupportedProfiles = 0;
@@ -123,15 +118,12 @@ namespace AVEncoder
 	};
 
 
-	struct UE_DEPRECATED(5.4, "AVEncoder has been deprecated. Please use the AVCodecs plugin family instead.") FVideoDecoderInfo
+	struct FVideoDecoderInfo
 	{
 		uint32						ID = 0;
+		ECodecType					CodecType = ECodecType::Undefined;
 		uint32						MaxWidth = 0;
 		uint32						MaxHeight = 0;
-
-		PRAGMA_DISABLE_DEPRECATION_WARNINGS
-		ECodecType					CodecType = ECodecType::Undefined;
-		PRAGMA_ENABLE_DEPRECATION_WARNINGS
 	};
 
 #if PLATFORM_WINDOWS
diff --git a/Engine/Source/Runtime/AVEncoder/Public/VideoDecoder.h b/Engine/Source/Runtime/AVEncoder/Public/VideoDecoder.h
--- a/Engine/Source/Runtime/AVEncoder/Public/VideoDecoder.h
+++ b/Engine/Source/Runtime/AVEncoder/Public/VideoDecoder.h
@@ -15,13 +15,13 @@ namespace AVEncoder
 	class FVideoDecoderInput;
 	class FVideoDecoderOutput;
 
-	class UE_DEPRECATED(5.4, "AVEncoder has been deprecated. Please use the AVCodecs plugin family instead.") FVideoDecoder
+	class FVideoDecoder
 	{
 	public:
 		using CreateDecoderAllocationInterfaceCallback = TFunction<void(void* /*InOptions*/, void** /*InOutParamResult*/)>;
 		using ReleaseDecoderAllocationInterfaceCallback = TFunction<void(void* /*InOptions*/, void** /*InOutParamResult*/)>;
 
-		struct UE_DEPRECATED(5.4, "AVEncoder has been deprecated. Please use the AVCodecs plugin family instead.") FInit
+		struct FInit
 		{
 			CreateDecoderAllocationInterfaceCallback	CreateDecoderAllocationInterface;
 			ReleaseDecoderAllocationInterfaceCallback	ReleaseDecoderAllocationInterface;
@@ -35,18 +35,15 @@ namespace AVEncoder
 		// Shuts the decoder down AND destroys it. Do not store a pointer to this decoder in a smart pointer!
 		virtual void Shutdown() = 0;
 
-		enum class UE_DEPRECATED(5.4, "AVEncoder has been deprecated. Please use the AVCodecs plugin family instead.") EDecodeResult
+		enum class EDecodeResult
 		{
 			Success,
 			Failure
 //			TryAgainLater
 		};
-
-		PRAGMA_DISABLE_DEPRECATION_WARNINGS
 		virtual EDecodeResult Decode(const FVideoDecoderInput* InInput) = 0;
 
 		using OnDecodedFrameCallback = TFunction<void(const FVideoDecoderOutput* /*InDecodedFrame*/)>;
-		PRAGMA_ENABLE_DEPRECATION_WARNINGS
 
 		virtual void SetOnDecodedFrame(OnDecodedFrameCallback InCallback) { OnDecodedFrame = MoveTemp(InCallback); }
 		virtual void ClearOnDecodedFrame() { OnDecodedFrame = nullptr; }
@@ -58,10 +55,8 @@ namespace AVEncoder
 	protected:
 		bool CreateDecoderAllocationInterface();
 		void ReleaseDecoderAllocationInterface();
-		void* GetAllocationInterfaceMethods();
-		PRAGMA_DISABLE_DEPRECATION_WARNINGS
 		EFrameBufferAllocReturn AllocateOutputFrameBuffer(FVideoDecoderAllocFrameBufferResult* OutBuffer, const FVideoDecoderAllocFrameBufferParams* InAllocParams);
-		PRAGMA_ENABLE_DEPRECATION_WARNINGS
+		void* GetAllocationInterfaceMethods();
 
 		struct FPlatformDecoderAllocInterface;
 
diff --git a/Engine/Source/Runtime/AVEncoder/Public/VideoDecoderAllocationTypes.h b/Engine/Source/Runtime/AVEncoder/Public/VideoDecoderAllocationTypes.h
--- a/Engine/Source/Runtime/AVEncoder/Public/VideoDecoderAllocationTypes.h
+++ b/Engine/Source/Runtime/AVEncoder/Public/VideoDecoderAllocationTypes.h
@@ -9,7 +9,7 @@
 /**
  * Type of video frame buffer to allocate.
  */
-enum class UE_DEPRECATED(5.4, "AVEncoder has been deprecated. Please use the AVCodecs plugin family instead.") EFrameBufferType
+enum class EFrameBufferType
 {
 	CODEC_RawBuffer = 0,
 	CODEC_TextureHandle = 1,
@@ -19,12 +19,10 @@ enum class UE_DEPRECATED(5.4, "AVEncoder has been deprecated. Please use the AVC
 /**
  * Video frame buffer allocation parameters.
  */
-struct UE_DEPRECATED(5.4, "AVEncoder has been deprecated. Please use the AVCodecs plugin family instead.") FVideoDecoderAllocFrameBufferParams
+struct FVideoDecoderAllocFrameBufferParams
 {
 	/** Type of output frame buffer to allocate */
-	PRAGMA_DISABLE_DEPRECATION_WARNINGS
 	EFrameBufferType FrameBufferType;
-	PRAGMA_ENABLE_DEPRECATION_WARNINGS
 	/**
 	 * Number of bytes to allocate with the specified alignment.
 	 * For a raw buffer the allocation size is required as it may include oversize for padding.
@@ -46,7 +44,7 @@ struct UE_DEPRECATED(5.4, "AVEncoder has been deprecated. Please use the AVCodec
 /**
  * Return value for the video decoder frame buffer allocation callback.
  */
-enum class UE_DEPRECATED(5.4, "AVEncoder has been deprecated. Please use the AVCodecs plugin family instead.") EFrameBufferAllocReturn
+enum class EFrameBufferAllocReturn
 {
 	/** Buffer successfully allocated. */
 	CODEC_Success = 0,
@@ -69,7 +67,7 @@ DECLARE_DELEGATE_TwoParams(FReleaseFrameBuffer, void* /* , This */, void* /*Buff
 /**
  * Output buffer plane description.
  */
-struct UE_DEPRECATED(5.4, "AVEncoder has been deprecated. Please use the AVCodecs plugin family instead.") FFrameBufferOutPlaneDesc
+struct FFrameBufferOutPlaneDesc
 {
 	/** Width of the allocated buffer plane in pixels, including necessary padding. */
 	int32_t Width;
@@ -88,7 +86,7 @@ struct UE_DEPRECATED(5.4, "AVEncoder has been deprecated. Please use the AVCodec
 /**
  * Result structure to be filled in by the application.
  */
-struct UE_DEPRECATED(5.4, "AVEncoder has been deprecated. Please use the AVCodecs plugin family instead.") FVideoDecoderAllocFrameBufferResult
+struct FVideoDecoderAllocFrameBufferResult
 {
 	/** The buffer the application has allocated. */
 	void* AllocatedBuffer;
@@ -99,9 +97,7 @@ struct UE_DEPRECATED(5.4, "AVEncoder has been deprecated. Please use the AVCodec
 	 *  to allocate). Supports at most 4 planes (R,G,B,A). Usually 2 (Y, UV) or 3 (Y,U,V or R,G,B). */
 	int32_t AllocatedPlanesNum;
 	int32_t AllocatedPlaneLayout;	/* reserved */
-	PRAGMA_DISABLE_DEPRECATION_WARNINGS
 	FFrameBufferOutPlaneDesc AllocatedPlaneDesc[4];
-	PRAGMA_ENABLE_DEPRECATION_WARNINGS
 
 	/** Callback within the application to retain the allocated buffer. */
 	FRetainFrameBuffer RetainCallback;
@@ -116,9 +112,7 @@ struct UE_DEPRECATED(5.4, "AVEncoder has been deprecated. Please use the AVCodec
  * Video frame buffer allocation callback within the application.
  * We call this to get a new output frame buffer to decode into.
  */
-PRAGMA_DISABLE_DEPRECATION_WARNINGS
 DECLARE_DELEGATE_RetVal_ThreeParams(EFrameBufferAllocReturn, FAllocFrameBuffer, void* /*This*/, const FVideoDecoderAllocFrameBufferParams* /*InAllocParams*/, FVideoDecoderAllocFrameBufferResult* /*OutBuffer*/);
-PRAGMA_ENABLE_DEPRECATION_WARNINGS
 
 
 
@@ -127,7 +121,7 @@ PRAGMA_ENABLE_DEPRECATION_WARNINGS
 
 DECLARE_DELEGATE_RetVal_ThreeParams(int32_t, FGetD3DDevice, void* /* , This */, void** /* , OutD3DDevice */, int32_t* /* , OutD3DVersion */);
 
-struct UE_DEPRECATED(5.4, "AVEncoder has been deprecated. Please use the AVCodecs plugin family instead.") FVideoDecoderMethodsWindows
+struct FVideoDecoderMethodsWindows
 {
 	/** Magic cookie to check for if this is indeed the expected structure. */
 	uint32_t MagicCookie;	// 'WinX' 0x57696e58
@@ -145,7 +139,7 @@ DECLARE_DELEGATE_RetVal_FiveParams(int32_t, FMemAlloc, void* /*, This*/, void**
 DECLARE_DELEGATE_RetVal_ThreeParams(int32_t, FMemFree, void* /* , This */, void** /* , InAddress */, int32_t /* , MemType */);
 DECLARE_DELEGATE_RetVal_ThreeParams(int32_t, FConfigQuery, void* /* , This */, void** /* , OutInfo */, void** /* , InInfo */);
 
-struct UE_DEPRECATED(5.4, "AVEncoder has been deprecated. Please use the AVCodecs plugin family instead.") FVideoDecoderMethods
+struct FVideoDecoderMethods
 {
 	/** Magic cookie to check for if this is indeed the expected structure. */
 	uint32_t MagicCookie;
diff --git a/Engine/Source/Runtime/AVEncoder/Public/VideoDecoderFactory.h b/Engine/Source/Runtime/AVEncoder/Public/VideoDecoderFactory.h
--- a/Engine/Source/Runtime/AVEncoder/Public/VideoDecoderFactory.h
+++ b/Engine/Source/Runtime/AVEncoder/Public/VideoDecoderFactory.h
@@ -15,7 +15,7 @@ namespace AVEncoder
 {
 	struct FVideoDecoderInfo;
 
-	class UE_DEPRECATED(5.4, "AVEncoder has been deprecated. Please use the AVCodecs plugin family instead.") FVideoDecoderFactory
+	class FVideoDecoderFactory
 	{
 	public:
 		static AVENCODER_API FVideoDecoderFactory& Get();
@@ -24,7 +24,6 @@ namespace AVEncoder
 		// --- decoder registry
 
 		// the callback type used to create a registered decoder
-		PRAGMA_DISABLE_DEPRECATION_WARNINGS
 		using CreateDecoderCallback = TFunction<FVideoDecoder*()>;
 
 		// register a decoder so that it can be iterated and created
@@ -37,7 +36,6 @@ namespace AVEncoder
 		// --- decoder creation
 		// create a decoder instance. If not nullptr it must be destroyed through invoking its Shutdown() method!
 		AVENCODER_API FVideoDecoder* Create(uint32 InID, const FVideoDecoder::FInit& InInit);
-		PRAGMA_ENABLE_DEPRECATION_WARNINGS
 
 	private:
 		AVENCODER_API FVideoDecoderFactory();
@@ -50,11 +48,8 @@ namespace AVEncoder
 		static AVENCODER_API FThreadSafeCounter		NextID;
 		static AVENCODER_API bool						bDebugDontRegisterDefaultCodecs;
 
-		TArray<CreateDecoderCallback>	CreateDecoders;
-
-		PRAGMA_DISABLE_DEPRECATION_WARNINGS
 		TArray<FVideoDecoderInfo>		AvailableDecoders;
-		PRAGMA_ENABLE_DEPRECATION_WARNINGS
+		TArray<CreateDecoderCallback>	CreateDecoders;
 	};
 
 } /* namespace AVEncoder */
diff --git a/Engine/Source/Runtime/AVEncoder/Public/VideoDecoderInput.h b/Engine/Source/Runtime/AVEncoder/Public/VideoDecoderInput.h
--- a/Engine/Source/Runtime/AVEncoder/Public/VideoDecoderInput.h
+++ b/Engine/Source/Runtime/AVEncoder/Public/VideoDecoderInput.h
@@ -14,11 +14,11 @@ struct FVideoDecoderAllocFrameBufferResult;
 namespace AVEncoder
 {
 
-	class UE_DEPRECATED(5.4, "AVEncoder has been deprecated. Please use the AVCodecs plugin family instead.") FVideoDecoderInput
+	class FVideoDecoderInput
 	{
 	public:
 
-		struct UE_DEPRECATED(5.4, "AVEncoder has been deprecated. Please use the AVCodecs plugin family instead.") FInputData
+		struct FInputData
 		{
 			const void* EncodedData = nullptr;
 			int64		PTS = 0;
@@ -54,7 +54,7 @@ namespace AVEncoder
 
 
 
-	class UE_DEPRECATED(5.4, "AVEncoder has been deprecated. Please use the AVCodecs plugin family instead.") FVideoDecoderOutput
+	class FVideoDecoderOutput
 	{
 	public:
 		virtual int32 AddRef() = 0;
@@ -63,9 +63,7 @@ namespace AVEncoder
 		virtual int32 GetHeight() const = 0;
 		virtual int64 GetPTS() const = 0;
 
-		PRAGMA_DISABLE_DEPRECATION_WARNINGS
 		virtual const FVideoDecoderAllocFrameBufferResult* GetAllocatedBuffer() const = 0;
-		PRAGMA_ENABLE_DEPRECATION_WARNINGS
 
 		virtual int32 GetCropLeft() const = 0;
 		virtual int32 GetCropRight() const = 0;
diff --git a/Engine/Source/Runtime/AVEncoder/Public/VideoEncoder.h b/Engine/Source/Runtime/AVEncoder/Public/VideoEncoder.h
--- a/Engine/Source/Runtime/AVEncoder/Public/VideoEncoder.h
+++ b/Engine/Source/Runtime/AVEncoder/Public/VideoEncoder.h
@@ -19,14 +19,14 @@ namespace AVEncoder
 	class FVideoEncoderInput;
 	class FVideoEncoderInputFrame;
 
-    class UE_DEPRECATED(5.4, "AVEncoder has been deprecated. Please use the AVCodecs plugin family instead.") FVideoEncoder
+    class FVideoEncoder
     {
     public:
         enum class RateControlMode { UNKNOWN, CONSTQP, VBR, CBR };
         enum class MultipassMode { UNKNOWN, DISABLED, QUARTER, FULL };
 		enum class H264Profile { AUTO, CONSTRAINED_BASELINE, BASELINE, MAIN, CONSTRAINED_HIGH, HIGH, HIGH444, STEREO, SVC_TEMPORAL_SCALABILITY, PROGRESSIVE_HIGH };
 
-        struct UE_DEPRECATED(5.4, "AVEncoder has been deprecated. Please use the AVCodecs plugin family instead.") FLayerConfig
+        struct FLayerConfig
         {
             uint32			Width = 0;
             uint32			Height = 0;
@@ -63,9 +63,7 @@ namespace AVEncoder
 
         AVENCODER_API virtual ~FVideoEncoder();
 
-        PRAGMA_DISABLE_DEPRECATION_WARNINGS
 		virtual bool Setup(TSharedRef<FVideoEncoderInput> input, FLayerConfig const& config) { return false; }
-        PRAGMA_ENABLE_DEPRECATION_WARNINGS
 		virtual void Shutdown() {}
 
         AVENCODER_API virtual bool AddLayer(FLayerConfig const& config);
@@ -75,12 +73,10 @@ namespace AVEncoder
 		AVENCODER_API FLayerConfig GetLayerConfig(uint32 layerIdx) const;
         AVENCODER_API void UpdateLayerConfig(uint32 layerIdx, FLayerConfig const& config);
 
-        PRAGMA_DISABLE_DEPRECATION_WARNINGS
         using OnFrameEncodedCallback = TFunction<void(const TSharedPtr<FVideoEncoderInputFrame> /* InCompletedFrame */)>;
         using OnEncodedPacketCallback = TFunction<void(uint32 /* LayerIndex */, const TSharedPtr<FVideoEncoderInputFrame> /* Frame */, const FCodecPacket& /* Packet */)>;
-        PRAGMA_ENABLE_DEPRECATION_WARNINGS
 
-        struct UE_DEPRECATED(5.4, "AVEncoder has been deprecated. Please use the AVCodecs plugin family instead.") FEncodeOptions
+        struct FEncodeOptions
         {
             bool					bForceKeyFrame = false;
             OnFrameEncodedCallback	OnFrameEncoded;
@@ -89,14 +85,12 @@ namespace AVEncoder
         void SetOnEncodedPacket(OnEncodedPacketCallback callback) { OnEncodedPacket = MoveTemp(callback); }
         void ClearOnEncodedPacket() { OnEncodedPacket = nullptr; }
 
-        PRAGMA_DISABLE_DEPRECATION_WARNINGS
 		virtual void Encode(const TSharedPtr<FVideoEncoderInputFrame> frame, FEncodeOptions const& options) {}
-        PRAGMA_ENABLE_DEPRECATION_WARNINGS
 
     protected:
         FVideoEncoder() = default;
 
-        class UE_DEPRECATED(5.4, "AVEncoder has been deprecated. Please use the AVCodecs plugin family instead.") FLayer
+        class FLayer
         {
         public:
             explicit FLayer(FLayerConfig const& layerConfig)
diff --git a/Engine/Source/Runtime/AVEncoder/Public/VideoEncoderFactory.h b/Engine/Source/Runtime/AVEncoder/Public/VideoEncoderFactory.h
--- a/Engine/Source/Runtime/AVEncoder/Public/VideoEncoderFactory.h
+++ b/Engine/Source/Runtime/AVEncoder/Public/VideoEncoderFactory.h
@@ -18,7 +18,7 @@ namespace AVEncoder
 {
 	class FVideoEncoderInput;
 
-	class UE_DEPRECATED(5.4, "AVEncoder has been deprecated. Please use the AVCodecs plugin family instead.") FVideoEncoderFactory
+	class FVideoEncoderFactory
 	{
 	public:
 		static AVENCODER_API FVideoEncoderFactory& Get();
@@ -30,7 +30,6 @@ namespace AVEncoder
 		// --- encoder registry
 
 		// the callback type used to create a registered encoder
-		PRAGMA_DISABLE_DEPRECATION_WARNINGS
 		using CreateEncoderCallback = TFunction<TUniquePtr<FVideoEncoder>()>;
 
 		// register an encoder so that it can be iterated and created
@@ -45,7 +44,7 @@ namespace AVEncoder
 
 		AVENCODER_API TUniquePtr<FVideoEncoder> Create(uint32 InID, const FVideoEncoder::FLayerConfig& config);
 		AVENCODER_API TUniquePtr<FVideoEncoder> Create(uint32 InID, TSharedPtr<FVideoEncoderInput> InInput, const FVideoEncoder::FLayerConfig& config);
-		PRAGMA_ENABLE_DEPRECATION_WARNINGS
+
 	private:
 		FVideoEncoderFactory() = default;
 		~FVideoEncoderFactory() = default;
@@ -59,11 +58,9 @@ namespace AVEncoder
 		static AVENCODER_API FThreadSafeCounter		NextID;
 		bool							bDebugDontRegisterDefaultCodecs = false;
 		FThreadSafeBool					bWasSetup;
+		TArray<FVideoEncoderInfo>		AvailableEncoders;
 		TArray<CreateEncoderCallback>	CreateEncoders;
 
-		PRAGMA_DISABLE_DEPRECATION_WARNINGS
-		TArray<FVideoEncoderInfo>		AvailableEncoders;
-		PRAGMA_ENABLE_DEPRECATION_WARNINGS
 	};
     
 } /* namespace AVEncoder */
diff --git a/Engine/Source/Runtime/AVEncoder/Public/VideoEncoderInput.h b/Engine/Source/Runtime/AVEncoder/Public/VideoEncoderInput.h
--- a/Engine/Source/Runtime/AVEncoder/Public/VideoEncoderInput.h
+++ b/Engine/Source/Runtime/AVEncoder/Public/VideoEncoderInput.h
@@ -34,7 +34,7 @@ namespace AVEncoder
 	class FVideoEncoderInputFrame;
 
 #if PLATFORM_DESKTOP && !PLATFORM_APPLE
-	struct UE_DEPRECATED(5.4, "AVEncoder has been deprecated. Please use the AVCodecs plugin family instead.") FVulkanDataStruct
+	struct FVulkanDataStruct
 	{
 		VkInstance VulkanInstance;
 		VkPhysicalDevice VulkanPhysicalDevice;
@@ -42,7 +42,7 @@ namespace AVEncoder
 	};
 #endif
 
-	class UE_DEPRECATED(5.4, "AVEncoder has been deprecated. Please use the AVCodecs plugin family instead.") FVideoEncoderInput
+	class FVideoEncoderInput
 	{
 	public:
 		// --- construct video encoder input based on expected input frame format
@@ -64,16 +64,12 @@ namespace AVEncoder
 		// --- properties
 		AVENCODER_API virtual void SetMaxNumBuffers(uint32 InMaxNumBuffers);
 
-		PRAGMA_DISABLE_DEPRECATION_WARNINGS
 		EVideoFrameFormat GetFrameFormat() const { return FrameFormat; }
-		PRAGMA_ENABLE_DEPRECATION_WARNINGS
 
 		// --- available encoders
 
 		// get a list of supported video encoders
-		PRAGMA_DISABLE_DEPRECATION_WARNINGS
 		virtual const TArray<FVideoEncoderInfo>& GetAvailableEncoders() = 0;
-		PRAGMA_ENABLE_DEPRECATION_WARNINGS
 
 		// --- create encoders
 
@@ -116,10 +112,7 @@ namespace AVEncoder
 		FVideoEncoderInput(const FVideoEncoderInput&) = delete;
 		FVideoEncoderInput& operator=(const FVideoEncoderInput&) = delete;
 
-		PRAGMA_DISABLE_DEPRECATION_WARNINGS
 		EVideoFrameFormat				FrameFormat = EVideoFrameFormat::Undefined;
-		PRAGMA_ENABLE_DEPRECATION_WARNINGS
-
 		uint32 MaxNumBuffers = 3;
 		uint32 NumBuffers = 0;
 
@@ -130,7 +123,7 @@ namespace AVEncoder
 
 
 	// TODO this should go elsewhere and be made cross platform
-	class UE_DEPRECATED(5.4, "AVEncoder has been deprecated. Please use the AVCodecs plugin family instead.") FVideoEncoderInputFrame
+	class FVideoEncoderInputFrame
 	{
 	public:
 		// Obtain (increase reference count) of this input frame
@@ -154,9 +147,7 @@ namespace AVEncoder
 		int64 GetTimestampRTP() const { return TimestampRTP; }
 
 		// current format of frame
-		PRAGMA_DISABLE_DEPRECATION_WARNINGS
 		EVideoFrameFormat GetFormat() const { return Format; }
-		PRAGMA_ENABLE_DEPRECATION_WARNINGS
 		// width of frame buffer
 		void SetWidth(uint32 InWidth) { Width = InWidth; }
 		uint32 GetWidth() const { return Width; }
@@ -168,7 +159,7 @@ namespace AVEncoder
 
 		// --- YUV420P
 
-		struct UE_DEPRECATED(5.4, "AVEncoder has been deprecated. Please use the AVCodecs plugin family instead.") FYUV420P
+		struct FYUV420P
 		{
 			const uint8*		Data[3] = { nullptr, nullptr, nullptr };
 			uint32				StrideY = 0;
@@ -189,7 +180,7 @@ namespace AVEncoder
 #if PLATFORM_WINDOWS
 		// --- D3D11
 
-		struct UE_DEPRECATED(5.4, "AVEncoder has been deprecated. Please use the AVCodecs plugin family instead.") FD3D11
+		struct FD3D11
 		{
 			ID3D11Texture2D*	Texture = nullptr;
 			ID3D11Device*		EncoderDevice = nullptr;
@@ -207,7 +198,7 @@ namespace AVEncoder
 
 		// --- D3D12
 
-		struct UE_DEPRECATED(5.4, "AVEncoder has been deprecated. Please use the AVCodecs plugin family instead.") FD3D12
+		struct FD3D12
 		{
 			ID3D12Resource*		Texture = nullptr;
 			ID3D12Device*		EncoderDevice = nullptr;
@@ -224,7 +215,7 @@ namespace AVEncoder
 
 #endif // PLATFORM_WINDOWS
 
-		enum class UE_DEPRECATED(5.4, "AVEncoder has been deprecated. Please use the AVCodecs plugin family instead.") EUnderlyingRHI
+		enum class EUnderlyingRHI
 		{
 			Undefined,
 			D3D11,
@@ -233,15 +224,12 @@ namespace AVEncoder
 		};
 
 		// --- CUDA
-		struct UE_DEPRECATED(5.4, "AVEncoder has been deprecated. Please use the AVCodecs plugin family instead.") FCUDA
+		struct FCUDA
 		{
 			CUarray			EncoderTexture = nullptr;
 			CUcontext   	EncoderDevice = nullptr;
-			void*			SharedHandle = nullptr;
-
-			PRAGMA_DISABLE_DEPRECATION_WARNINGS
 			EUnderlyingRHI	UnderlyingRHI = EUnderlyingRHI::Undefined;
-			PRAGMA_ENABLE_DEPRECATION_WARNINGS
+			void*			SharedHandle = nullptr;
 		};
 
 		const FCUDA& GetCUDA() const { return CUDA; }
@@ -254,7 +242,7 @@ namespace AVEncoder
 
 #if PLATFORM_DESKTOP && !PLATFORM_APPLE
 		// --- Vulkan
-		struct UE_DEPRECATED(5.4, "AVEncoder has been deprecated. Please use the AVCodecs plugin family instead.") FVulkan
+		struct FVulkan
 		{
 			VkImage				EncoderTexture = VK_NULL_HANDLE;
 			VkDeviceMemory		EncoderDeviceMemory;
@@ -285,15 +273,12 @@ namespace AVEncoder
 		int64									TimestampUs;
 		int64									TimestampRTP;
 		mutable FThreadSafeCounter				NumReferences;
+		EVideoFrameFormat						Format;
 		uint32									Width;
 		uint32									Height;
 		FYUV420P								YUV420P;
 		bool									bFreeYUV420PData;
 
-		PRAGMA_DISABLE_DEPRECATION_WARNINGS
-		EVideoFrameFormat						Format;
-		PRAGMA_ENABLE_DEPRECATION_WARNINGS
-
 #if PLATFORM_WINDOWS
 		FD3D11									D3D11;
 		FReleaseD3D11TextureCallback			OnReleaseD3D11Texture;
diff --git a/Engine/Source/Runtime/GameplayMediaEncoder/Private/GameplayMediaEncoder.cpp b/Engine/Source/Runtime/GameplayMediaEncoder/Private/GameplayMediaEncoder.cpp
--- a/Engine/Source/Runtime/GameplayMediaEncoder/Private/GameplayMediaEncoder.cpp
+++ b/Engine/Source/Runtime/GameplayMediaEncoder/Private/GameplayMediaEncoder.cpp
@@ -49,15 +49,14 @@ const uint32 MaxVideoFPS = 60;
 const uint32 MaxWidth = 1920;
 const uint32 MaxHeight = 1080;
 
-PRAGMA_DISABLE_DEPRECATION_WARNINGS
-FAutoConsoleCommand GameplayMediaEncoderInitialize(TEXT("GameplayMediaEncoder.Initialize"), TEXT("Constructs the audio/video encoding objects. Does not start encoding"), FConsoleCommandDelegate::CreateStatic(&FGameplayMediaEncoder::InitializeCmd));
+FAutoConsoleCommand GameplayMediaEncoderInitialize(TEXT("GameplayMediaEncoder.Initialize"), TEXT("Constructs the audio/video encoding objects. Does not start encoding"),
+                                                   FConsoleCommandDelegate::CreateStatic(&FGameplayMediaEncoder::InitializeCmd));
 
 FAutoConsoleCommand GameplayMediaEncoderStart(TEXT("GameplayMediaEncoder.Start"), TEXT("Starts encoding"), FConsoleCommandDelegate::CreateStatic(&FGameplayMediaEncoder::StartCmd));
 
 FAutoConsoleCommand GameplayMediaEncoderStop(TEXT("GameplayMediaEncoder.Stop"), TEXT("Stops encoding"), FConsoleCommandDelegate::CreateStatic(&FGameplayMediaEncoder::StopCmd));
 
 FAutoConsoleCommand GameplayMediaEncoderShutdown(TEXT("GameplayMediaEncoder.Shutdown"), TEXT("Releases all systems."), FConsoleCommandDelegate::CreateStatic(&FGameplayMediaEncoder::ShutdownCmd));
-PRAGMA_ENABLE_DEPRECATION_WARNINGS
 
 //////////////////////////////////////////////////////////////////////////
 //
@@ -65,32 +64,21 @@ PRAGMA_ENABLE_DEPRECATION_WARNINGS
 //
 //////////////////////////////////////////////////////////////////////////
 
-PRAGMA_DISABLE_DEPRECATION_WARNINGS
 TSharedPtr<FGameplayMediaEncoder, ESPMode::ThreadSafe> FGameplayMediaEncoder::Singleton = { };
-PRAGMA_ENABLE_DEPRECATION_WARNINGS
 
-PRAGMA_DISABLE_DEPRECATION_WARNINGS
 FGameplayMediaEncoder* FGameplayMediaEncoder::Get()
-PRAGMA_ENABLE_DEPRECATION_WARNINGS
 {
 	// Constructed and captured as a thread safe shared pointer as this is required by the audio submix listener interface
 	if(!Singleton.IsValid())
 	{
-		PRAGMA_DISABLE_DEPRECATION_WARNINGS
 		Singleton = MakeShared<FGameplayMediaEncoder>();
-		PRAGMA_ENABLE_DEPRECATION_WARNINGS
 	}
-	
 	return Singleton.Get();
 }
 
-PRAGMA_DISABLE_DEPRECATION_WARNINGS
 FGameplayMediaEncoder::FGameplayMediaEncoder() {}
-PRAGMA_ENABLE_DEPRECATION_WARNINGS
 
-PRAGMA_DISABLE_DEPRECATION_WARNINGS
 FGameplayMediaEncoder::~FGameplayMediaEncoder() { Shutdown(); }
-PRAGMA_ENABLE_DEPRECATION_WARNINGS
 
 bool FGameplayMediaEncoder::RegisterListener(IGameplayMediaEncoderListener* Listener)
 {
@@ -149,25 +137,20 @@ bool FGameplayMediaEncoder::Initialize()
 	//
 	// Audio
 	//
-	PRAGMA_DISABLE_DEPRECATION_WARNINGS
 	AVEncoder::FAudioEncoderFactory* AudioEncoderFactory = AVEncoder::FAudioEncoderFactory::FindFactory("aac");
-	PRAGMA_ENABLE_DEPRECATION_WARNINGS
 	if(!AudioEncoderFactory)
 	{
 		UE_LOG(GameplayMediaEncoder, Error, TEXT("No audio encoder for aac found"));
 		return false;
 	}
 
-	PRAGMA_DISABLE_DEPRECATION_WARNINGS
 	AudioEncoder = AudioEncoderFactory->CreateEncoder("aac");
-	PRAGMA_ENABLE_DEPRECATION_WARNINGS
 	if(!AudioEncoder)
 	{
 		UE_LOG(GameplayMediaEncoder, Error, TEXT("Could not create audio encoder"));
 		return false;
 	}
 
-	PRAGMA_DISABLE_DEPRECATION_WARNINGS
 	AVEncoder::FAudioConfig AudioConfig;
 	AudioConfig.Samplerate = HardcodedAudioSamplerate;
 	AudioConfig.NumChannels = HardcodedAudioNumChannels;
@@ -179,14 +162,13 @@ bool FGameplayMediaEncoder::Initialize()
 	}
 
 	AudioEncoder->RegisterListener(*this);
-	PRAGMA_ENABLE_DEPRECATION_WARNINGS
 
 	MemoryCheckpoint("Audio encoder initialized");
 
 	//
 	// Video
 	//
-	PRAGMA_DISABLE_DEPRECATION_WARNINGS
+
 	VideoConfig.Codec = "h264";
 	VideoConfig.Height = VideoConfig.Width = VideoConfig.Framerate = VideoConfig.Bitrate = 0;
 	FParse::Value(FCommandLine::Get(), TEXT("GameplayMediaEncoder.ResY="), VideoConfig.Height);
@@ -201,7 +183,6 @@ bool FGameplayMediaEncoder::Initialize()
 		VideoConfig.Width = 1920;
 		VideoConfig.Height = 1080;
 	}
-	PRAGMA_ENABLE_DEPRECATION_WARNINGS
 	else
 	{
 		UE_LOG(GameplayMediaEncoder, Fatal, TEXT("GameplayMediaEncoder.ResY can only have a value of 720 or 1080"));
@@ -209,29 +190,22 @@ bool FGameplayMediaEncoder::Initialize()
 	}
 
 	// Specifying 0 will completely disable frame skipping (therefore encoding as many frames as possible)
-	PRAGMA_DISABLE_DEPRECATION_WARNINGS
 	FParse::Value(FCommandLine::Get(), TEXT("GameplayMediaEncoder.FPS="), VideoConfig.Framerate);
 	if(VideoConfig.Framerate == 0)
-	PRAGMA_ENABLE_DEPRECATION_WARNINGS
 	{
 		// Note : When disabling frame skipping, we lie to the encoder when initializing.
 		// We still specify a framerate, but then feed frames without skipping
-		PRAGMA_DISABLE_DEPRECATION_WARNINGS
 		VideoConfig.Framerate = HardcodedVideoFPS;
-		PRAGMA_ENABLE_DEPRECATION_WARNINGS
 		bDoFrameSkipping = false;
 		UE_LOG(GameplayMediaEncoder, Log, TEXT("Uncapping FPS"));
 	}
 	else
 	{
-		PRAGMA_DISABLE_DEPRECATION_WARNINGS
 		VideoConfig.Framerate = FMath::Clamp(VideoConfig.Framerate, (uint32)MinVideoFPS, (uint32)MaxVideoFPS);
-		UE_LOG(GameplayMediaEncoder, Log, TEXT("Capping FPS %u"), VideoConfig.Framerate);
-		PRAGMA_ENABLE_DEPRECATION_WARNINGS
 		bDoFrameSkipping = true;
+		UE_LOG(GameplayMediaEncoder, Log, TEXT("Capping FPS %u"), VideoConfig.Framerate);
 	}
 
-	PRAGMA_DISABLE_DEPRECATION_WARNINGS
 	VideoConfig.Bitrate = HardcodedVideoBitrate;
 	FParse::Value(FCommandLine::Get(), TEXT("GameplayMediaEncoder.Bitrate="), VideoConfig.Bitrate);
 	VideoConfig.Bitrate = FMath::Clamp(VideoConfig.Bitrate, (uint32)MinVideoBitrate, (uint32)MaxVideoBitrate);
@@ -242,7 +216,6 @@ bool FGameplayMediaEncoder::Initialize()
 	videoInit.MaxBitrate = MaxVideoBitrate;
 	videoInit.TargetBitrate = VideoConfig.Bitrate;
 	videoInit.MaxFramerate = VideoConfig.Framerate;
-	PRAGMA_ENABLE_DEPRECATION_WARNINGS
 
 	if(GDynamicRHI)
 	{
@@ -251,25 +224,16 @@ bool FGameplayMediaEncoder::Initialize()
 #if PLATFORM_DESKTOP && !PLATFORM_APPLE
 		if (RHIType == ERHIInterfaceType::D3D11)
 		{
-			PRAGMA_DISABLE_DEPRECATION_WARNINGS
 			VideoEncoderInput = AVEncoder::FVideoEncoderInput::CreateForD3D11(GDynamicRHI->RHIGetNativeDevice(), true, IsRHIDeviceAMD());
-			PRAGMA_ENABLE_DEPRECATION_WARNINGS
 		}
-		
 		else if (RHIType == ERHIInterfaceType::D3D12)
 		{
-			PRAGMA_DISABLE_DEPRECATION_WARNINGS
 			VideoEncoderInput = AVEncoder::FVideoEncoderInput::CreateForD3D12(GDynamicRHI->RHIGetNativeDevice(), true, IsRHIDeviceNVIDIA());
-			PRAGMA_ENABLE_DEPRECATION_WARNINGS
 		}
-		
 		else if (RHIType == ERHIInterfaceType::Vulkan)
 		{
-			PRAGMA_DISABLE_DEPRECATION_WARNINGS
 			VideoEncoderInput = AVEncoder::FVideoEncoderInput::CreateForVulkan(GDynamicRHI->RHIGetNativeDevice(), true);
-			PRAGMA_ENABLE_DEPRECATION_WARNINGS
 		}
-		
 		else
 #endif
 		{
@@ -278,9 +242,7 @@ bool FGameplayMediaEncoder::Initialize()
 		}
 	}
 
-	PRAGMA_DISABLE_DEPRECATION_WARNINGS
 	const TArray<AVEncoder::FVideoEncoderInfo>& AvailableEncodersInfo = AVEncoder::FVideoEncoderFactory::Get().GetAvailable();
-	PRAGMA_ENABLE_DEPRECATION_WARNINGS
 
 	if (AvailableEncodersInfo.Num() == 0)
 	{
@@ -290,12 +252,10 @@ bool FGameplayMediaEncoder::Initialize()
 
 	for (const auto& EncoderInfo : AvailableEncodersInfo)
 	{
-		PRAGMA_DISABLE_DEPRECATION_WARNINGS
 		if (EncoderInfo.CodecType == AVEncoder::ECodecType::H264)
 		{
 			VideoEncoder = AVEncoder::FVideoEncoderFactory::Get().Create(EncoderInfo.ID, VideoEncoderInput, videoInit);
 		}
-		PRAGMA_ENABLE_DEPRECATION_WARNINGS
 	}
 
 	if (!VideoEncoder)
@@ -304,10 +264,8 @@ bool FGameplayMediaEncoder::Initialize()
 		return false;
 	}
 
-	PRAGMA_DISABLE_DEPRECATION_WARNINGS
 	VideoEncoder->SetOnEncodedPacket([this](uint32 LayerIndex, const TSharedPtr<AVEncoder::FVideoEncoderInputFrame> Frame, const AVEncoder::FCodecPacket& Packet)
 	                                 { OnEncodedVideoFrame(LayerIndex, Frame, Packet); });
-	PRAGMA_ENABLE_DEPRECATION_WARNINGS
 
 	if(!VideoEncoder)
 	{
@@ -353,9 +311,7 @@ bool FGameplayMediaEncoder::Start()
 		AudioDevice->RegisterSubmixBufferListener(AsShared(), AudioDevice->GetMainSubmixObject());
 	}
 
-	PRAGMA_DISABLE_DEPRECATION_WARNINGS
 	FSlateApplication::Get().GetRenderer()->OnBackBufferReadyToPresent().AddRaw(this, &FGameplayMediaEncoder::OnFrameBufferReady);
-	PRAGMA_ENABLE_DEPRECATION_WARNINGS
 
 	return true;
 }
@@ -387,11 +343,8 @@ void FGameplayMediaEncoder::Stop()
 	AudioClock = 0;
 }
 
-PRAGMA_DISABLE_DEPRECATION_WARNINGS
 AVEncoder::FAudioConfig FGameplayMediaEncoder::GetAudioConfig() const
-PRAGMA_ENABLE_DEPRECATION_WARNINGS
 {
-	PRAGMA_DISABLE_DEPRECATION_WARNINGS
 	if(AudioEncoder)
 	{
 		auto codec = AudioEncoder->GetType();
@@ -402,7 +355,6 @@ PRAGMA_ENABLE_DEPRECATION_WARNINGS
 	{
 		return {};
 	}
-	PRAGMA_ENABLE_DEPRECATION_WARNINGS
 }
 
 void FGameplayMediaEncoder::Shutdown()
@@ -418,9 +370,7 @@ void FGameplayMediaEncoder::Shutdown()
 		if(AudioEncoder)
 		{
 			// AudioEncoder->Reset();
-			PRAGMA_DISABLE_DEPRECATION_WARNINGS
 			AudioEncoder->Shutdown();
-			PRAGMA_ENABLE_DEPRECATION_WARNINGS
 			AudioEncoder.Reset();
 		}
 	}
@@ -428,9 +378,7 @@ void FGameplayMediaEncoder::Shutdown()
 		FScopeLock Lock(&VideoProcessingCS);
 		if(VideoEncoder)
 		{
-			PRAGMA_DISABLE_DEPRECATION_WARNINGS
 			VideoEncoder->Shutdown();
-			PRAGMA_ENABLE_DEPRECATION_WARNINGS
 			VideoEncoder.Reset();
 
 			BackBuffers.Empty();
@@ -440,9 +388,7 @@ void FGameplayMediaEncoder::Shutdown()
 
 FTimespan FGameplayMediaEncoder::GetMediaTimestamp() const { return FTimespan::FromSeconds(FPlatformTime::Seconds()) - StartTime; }
 
-PRAGMA_DISABLE_DEPRECATION_WARNINGS
 const FString& FGameplayMediaEncoder::GetListenerName() const
-PRAGMA_ENABLE_DEPRECATION_WARNINGS
 {
 	static const FString& ListenerName = TEXT("GameplayMediaEncoderListener");
 	return ListenerName;
@@ -585,14 +531,12 @@ void FGameplayMediaEncoder::ProcessAudioFrame(const float* AudioData, int32 NumS
 	// auto encodedInfo = AudioEncoder->Encode(timestamp, audio, &encoded);
 	// OnEncodedAudioFrame(encodedInfo, &encoded);
 
-	PRAGMA_DISABLE_DEPRECATION_WARNINGS
 	AVEncoder::FAudioFrame Frame;
 	Frame.Timestamp = FTimespan::FromSeconds(AudioClock);
 	Frame.Duration = FTimespan::FromSeconds(FloatBuffer.GetSampleDuration());
 	FloatBuffer.Clamp();
 	Frame.Data = FloatBuffer;
 	AudioEncoder->Encode(Frame);
-	PRAGMA_ENABLE_DEPRECATION_WARNINGS
 
 	AudioClock += FloatBuffer.GetSampleDuration();
 }
@@ -611,9 +555,7 @@ void FGameplayMediaEncoder::ProcessVideoFrame(const FTextureRHIRef& FrameBuffer)
 
 	if(bDoFrameSkipping)
 	{
-		PRAGMA_DISABLE_DEPRECATION_WARNINGS
 		uint64 NumExpectedFrames = static_cast<uint64>(Now.GetTotalSeconds() * VideoConfig.Framerate);
-		PRAGMA_ENABLE_DEPRECATION_WARNINGS
 		UE_LOG(GameplayMediaEncoder, VeryVerbose, TEXT("time %.3f: captured %d, expected %d"), Now.GetTotalSeconds(), NumCapturedFrames + 1, NumExpectedFrames);
 		if(NumCapturedFrames + 1 > NumExpectedFrames)
 		{
@@ -624,32 +566,24 @@ void FGameplayMediaEncoder::ProcessVideoFrame(const FTextureRHIRef& FrameBuffer)
 
 	UpdateVideoConfig();
 
-	PRAGMA_DISABLE_DEPRECATION_WARNINGS
 	TSharedPtr<AVEncoder::FVideoEncoderInputFrame> InputFrame = ObtainInputFrame();
 	const int32 FrameId = InputFrame->GetFrameID();
 	InputFrame->SetTimestampUs(Now.GetTicks());
-	PRAGMA_ENABLE_DEPRECATION_WARNINGS
 
 	CopyTexture(FrameBuffer, BackBuffers[InputFrame]);
 
-	PRAGMA_DISABLE_DEPRECATION_WARNINGS
 	AVEncoder::FVideoEncoder::FEncodeOptions EncodeOptions;
 	VideoEncoder->Encode(InputFrame, EncodeOptions);
-	PRAGMA_ENABLE_DEPRECATION_WARNINGS
 
 	LastVideoInputTimestamp = Now;
 	NumCapturedFrames++;
 }
 
-PRAGMA_DISABLE_DEPRECATION_WARNINGS
 TSharedPtr<AVEncoder::FVideoEncoderInputFrame> FGameplayMediaEncoder::ObtainInputFrame()
-PRAGMA_ENABLE_DEPRECATION_WARNINGS
 {
-	PRAGMA_DISABLE_DEPRECATION_WARNINGS
 	TSharedPtr<AVEncoder::FVideoEncoderInputFrame> InputFrame = VideoEncoderInput->ObtainInputFrame();
 	InputFrame->SetWidth(VideoConfig.Width);
 	InputFrame->SetHeight(VideoConfig.Height);
-	PRAGMA_ENABLE_DEPRECATION_WARNINGS
 
 	if(!BackBuffers.Contains(InputFrame))
 	{
@@ -657,9 +591,7 @@ PRAGMA_ENABLE_DEPRECATION_WARNINGS
 		const ERHIInterfaceType RHIType = RHIGetInterfaceType();
 
 		const FRHITextureCreateDesc Desc =
-			PRAGMA_DISABLE_DEPRECATION_WARNINGS
 			FRHITextureCreateDesc::Create2D(TEXT("VideoCapturerBackBuffer"), VideoConfig.Width, VideoConfig.Height, PF_B8G8R8A8)
-			PRAGMA_ENABLE_DEPRECATION_WARNINGS
 			.SetFlags(ETextureCreateFlags::Shared | ETextureCreateFlags::RenderTargetable | ETextureCreateFlags::UAV)
 			.SetInitialState(ERHIAccess::CopyDest);
 
@@ -668,9 +600,7 @@ PRAGMA_ENABLE_DEPRECATION_WARNINGS
 			FRHIResourceCreateInfo CreateInfo(TEXT("VideoCapturerBackBuffer"));
 			FTextureRHIRef Texture = RHICreateTexture(Desc);
 
-			PRAGMA_DISABLE_DEPRECATION_WARNINGS
 			InputFrame->SetTexture((ID3D11Texture2D*)Texture->GetNativeResource(), [&, InputFrame](ID3D11Texture2D* NativeTexture) { BackBuffers.Remove(InputFrame); });
-			PRAGMA_ENABLE_DEPRECATION_WARNINGS
 			BackBuffers.Add(InputFrame, Texture);
 		}
 		else if (RHIType == ERHIInterfaceType::D3D12)
@@ -678,9 +608,7 @@ PRAGMA_ENABLE_DEPRECATION_WARNINGS
 			FRHIResourceCreateInfo CreateInfo(TEXT("VideoCapturerBackBuffer"));
 			FTextureRHIRef Texture = RHICreateTexture(Desc);
 
-			PRAGMA_DISABLE_DEPRECATION_WARNINGS
 			InputFrame->SetTexture((ID3D12Resource*)Texture->GetNativeResource(), [&, InputFrame](ID3D12Resource* NativeTexture) { BackBuffers.Remove(InputFrame); });
-			PRAGMA_ENABLE_DEPRECATION_WARNINGS
 			BackBuffers.Add(InputFrame, Texture);
 		}
 
@@ -709,37 +637,27 @@ void FGameplayMediaEncoder::UpdateVideoConfig()
 {
 	if(bChangeBitrate || bChangeFramerate)
 	{
-		PRAGMA_DISABLE_DEPRECATION_WARNINGS
 		auto config = VideoEncoder->GetLayerConfig(0);
-		PRAGMA_ENABLE_DEPRECATION_WARNINGS
 
 		if(bChangeBitrate)
 		{
-			PRAGMA_DISABLE_DEPRECATION_WARNINGS
 			config.MaxBitrate = MaxVideoBitrate;
 			config.TargetBitrate = NewVideoBitrate;
-			PRAGMA_ENABLE_DEPRECATION_WARNINGS
 		}
 
 		if(bChangeFramerate)
 		{
-			PRAGMA_DISABLE_DEPRECATION_WARNINGS
 			config.MaxFramerate = NewVideoFramerate;
-			PRAGMA_ENABLE_DEPRECATION_WARNINGS
 			NumCapturedFrames = 0;
 		}
 
-		PRAGMA_DISABLE_DEPRECATION_WARNINGS
 		VideoEncoder->UpdateLayerConfig(0, config);
-		PRAGMA_ENABLE_DEPRECATION_WARNINGS
 		bChangeFramerate = false;
 		bChangeBitrate = false;
 	}
 }
 
-PRAGMA_DISABLE_DEPRECATION_WARNINGS
 void FGameplayMediaEncoder::OnEncodedAudioFrame(const AVEncoder::FMediaPacket& Packet)
-PRAGMA_ENABLE_DEPRECATION_WARNINGS
 {
 
 	FScopeLock Lock(&ListenersCS);
@@ -749,11 +667,8 @@ PRAGMA_ENABLE_DEPRECATION_WARNINGS
 	}
 }
 
-PRAGMA_DISABLE_DEPRECATION_WARNINGS
 void FGameplayMediaEncoder::OnEncodedVideoFrame(uint32 LayerIndex, const TSharedPtr<AVEncoder::FVideoEncoderInputFrame> InputFrame, const AVEncoder::FCodecPacket& Packet)
-PRAGMA_ENABLE_DEPRECATION_WARNINGS
 {
-	PRAGMA_DISABLE_DEPRECATION_WARNINGS
 	AVEncoder::FMediaPacket packet(AVEncoder::EPacketType::Video);
 
 	packet.Timestamp = InputFrame->GetTimestampUs();
@@ -764,7 +679,6 @@ PRAGMA_ENABLE_DEPRECATION_WARNINGS
 	packet.Video.Height = InputFrame->GetHeight();
 	packet.Video.FrameAvgQP = Packet.VideoQP;
 	packet.Video.Framerate = VideoConfig.Framerate;
-	PRAGMA_ENABLE_DEPRECATION_WARNINGS
 
 	FScopeLock Lock(&ListenersCS);
 	for(auto&& Listener : Listeners)
@@ -772,9 +686,7 @@ PRAGMA_ENABLE_DEPRECATION_WARNINGS
 		Listener->OnMediaSample(packet);
 	}
 
-	PRAGMA_DISABLE_DEPRECATION_WARNINGS
 	InputFrame->Release();
-	PRAGMA_ENABLE_DEPRECATION_WARNINGS
 }
 
 void FGameplayMediaEncoder::CopyTexture(const FTextureRHIRef& SourceTexture, FTextureRHIRef& DestinationTexture) const
diff --git a/Engine/Source/Runtime/GameplayMediaEncoder/Private/GameplayMediaEncoderModule.cpp b/Engine/Source/Runtime/GameplayMediaEncoder/Private/GameplayMediaEncoderModule.cpp
--- a/Engine/Source/Runtime/GameplayMediaEncoder/Private/GameplayMediaEncoderModule.cpp
+++ b/Engine/Source/Runtime/GameplayMediaEncoder/Private/GameplayMediaEncoderModule.cpp
@@ -27,11 +27,8 @@ public:
 		// If the FGameplayMediaEncoder instance was created, then explicitly destroy it here
 		// instead of waiting for the automatic cleanup, since at that point some objects
 		// it depends to for a clean shutdown are not available any longer.
-		PRAGMA_DISABLE_DEPRECATION_WARNINGS
 		FGameplayMediaEncoder::Singleton.Reset();
-		PRAGMA_ENABLE_DEPRECATION_WARNINGS
 	}
-	
 private:
 };
 
diff --git a/Engine/Source/Runtime/GameplayMediaEncoder/Public/GameplayMediaEncoder.h b/Engine/Source/Runtime/GameplayMediaEncoder/Public/GameplayMediaEncoder.h
--- a/Engine/Source/Runtime/GameplayMediaEncoder/Public/GameplayMediaEncoder.h
+++ b/Engine/Source/Runtime/GameplayMediaEncoder/Public/GameplayMediaEncoder.h
@@ -32,14 +32,10 @@ class SWindow;
 class IGameplayMediaEncoderListener
 {
 public:
-	PRAGMA_DISABLE_DEPRECATION_WARNINGS
 	virtual void OnMediaSample(const AVEncoder::FMediaPacket& Sample) = 0;
-	PRAGMA_ENABLE_DEPRECATION_WARNINGS
 };
 
-PRAGMA_DISABLE_DEPRECATION_WARNINGS
-class UE_DEPRECATED(5.4, "GameplayMediaEncoder will move to a plugin in a future Unreal Engine release.") FGameplayMediaEncoder final : public ISubmixBufferListener, public AVEncoder::IAudioEncoderListener
-PRAGMA_ENABLE_DEPRECATION_WARNINGS
+class FGameplayMediaEncoder final : public ISubmixBufferListener, public AVEncoder::IAudioEncoderListener
 {
 public:
 
@@ -91,10 +87,8 @@ public:
 		Get()->Stop();
 	}
 
-	PRAGMA_DISABLE_DEPRECATION_WARNINGS
 	GAMEPLAYMEDIAENCODER_API AVEncoder::FAudioConfig GetAudioConfig() const;
 	AVEncoder::FVideoConfig GetVideoConfig() const { return VideoConfig; }
-	PRAGMA_ENABLE_DEPRECATION_WARNINGS
 
 private:
 	template <typename ObjectType, ESPMode Mode>
@@ -118,12 +112,10 @@ private:
 
 	GAMEPLAYMEDIAENCODER_API void UpdateVideoConfig();
 
-	PRAGMA_DISABLE_DEPRECATION_WARNINGS
 	GAMEPLAYMEDIAENCODER_API void OnEncodedAudioFrame(const AVEncoder::FMediaPacket& Packet) override;
 	GAMEPLAYMEDIAENCODER_API void OnEncodedVideoFrame(uint32 LayerIndex, const TSharedPtr<AVEncoder::FVideoEncoderInputFrame> Frame, const AVEncoder::FCodecPacket& Packet);
 
 	GAMEPLAYMEDIAENCODER_API TSharedPtr<AVEncoder::FVideoEncoderInputFrame> ObtainInputFrame();
-	PRAGMA_ENABLE_DEPRECATION_WARNINGS
 	GAMEPLAYMEDIAENCODER_API void CopyTexture(const FTextureRHIRef& SourceTexture, FTextureRHIRef& DestinationTexture) const;
 
 	GAMEPLAYMEDIAENCODER_API void FloatToPCM16(float const* floatSamples, int32 numSamples, TArray<int16>& out) const;
@@ -134,14 +126,12 @@ private:
 	FCriticalSection AudioProcessingCS;
 	FCriticalSection VideoProcessingCS;
 
-	PRAGMA_DISABLE_DEPRECATION_WARNINGS
 	TUniquePtr<AVEncoder::FAudioEncoder> AudioEncoder;
 
 	AVEncoder::FVideoConfig VideoConfig;
 
 	TUniquePtr<AVEncoder::FVideoEncoder> VideoEncoder;
 	TSharedPtr<AVEncoder::FVideoEncoderInput> VideoEncoderInput;
-	PRAGMA_ENABLE_DEPRECATION_WARNINGS
 
 	uint64 NumCapturedFrames = 0;
 	FTimespan StartTime = 0;
@@ -171,4 +161,3 @@ private:
 	PRAGMA_ENABLE_DEPRECATION_WARNINGS
 };
 
-
diff --git a/Engine/Source/Runtime/Windows/WindowsPlatformFeatures/Private/HighlightRecorder.cpp b/Engine/Source/Runtime/Windows/WindowsPlatformFeatures/Private/HighlightRecorder.cpp
--- a/Engine/Source/Runtime/Windows/WindowsPlatformFeatures/Private/HighlightRecorder.cpp
+++ b/Engine/Source/Runtime/Windows/WindowsPlatformFeatures/Private/HighlightRecorder.cpp
@@ -66,9 +66,7 @@ bool FHighlightRecorder::Start(double RingBufferDurationSecs)
 	TotalPausedDuration = 0;
 	NumPushedFrames = 0;
 
-	PRAGMA_DISABLE_DEPRECATION_WARNINGS
 	if (!FGameplayMediaEncoder::Get()->RegisterListener(this))
-	PRAGMA_ENABLE_DEPRECATION_WARNINGS
 	{
 		return false;
 	}
@@ -114,9 +112,7 @@ void FHighlightRecorder::Stop()
 {
 	CSV_SCOPED_TIMING_STAT(WindowsVideoRecordingSystem, HighlightRecorder_Stop);
 
-	PRAGMA_DISABLE_DEPRECATION_WARNINGS
 	FGameplayMediaEncoder::Get()->UnregisterListener(this);
-	PRAGMA_ENABLE_DEPRECATION_WARNINGS
 	State = EState::Stopped;
 
 	if (BackgroundSaving)
@@ -133,7 +129,6 @@ FTimespan FHighlightRecorder::GetRecordingTime() const
 	return FTimespan::FromSeconds(FPlatformTime::Seconds()) - RecordingStartTime - TotalPausedDuration;
 }
 
-PRAGMA_DISABLE_DEPRECATION_WARNINGS
 void FHighlightRecorder::OnMediaSample(const AVEncoder::FMediaPacket& InSample)
 {
 	// We might be paused, so don't do anything
@@ -170,7 +165,6 @@ void FHighlightRecorder::OnMediaSample(const AVEncoder::FMediaPacket& InSample)
 
 	RingBuffer.Push(MoveTemp(SampleCopy));
 }
-PRAGMA_ENABLE_DEPRECATION_WARNINGS
 
 bool FHighlightRecorder::SaveHighlight(const TCHAR* Filename, FDoneCallback InDoneCallback, double MaxDurationSecs)
 {
@@ -232,13 +226,9 @@ bool FHighlightRecorder::SaveHighlightInBackground(const FString& Filename, doub
 	return bRes;
 }
 
-
 bool FHighlightRecorder::SaveHighlightInBackgroundImpl(const FString& Filename, double MaxDurationSecs)
 {
-	PRAGMA_DISABLE_DEPRECATION_WARNINGS
 	TArray<AVEncoder::FMediaPacket> Samples = RingBuffer.GetCopy();
-	PRAGMA_ENABLE_DEPRECATION_WARNINGS
-	
 
 	if (Samples.Num()==0)
 	{
@@ -257,9 +247,7 @@ bool FHighlightRecorder::SaveHighlightInBackgroundImpl(const FString& Filename,
 	bool bHasAudio = false;
 	for (int Idx = FirstSampleIndex; Idx != Samples.Num(); ++Idx)
 	{
-		PRAGMA_DISABLE_DEPRECATION_WARNINGS
 		if (Samples[Idx].Type == AVEncoder::EPacketType::Audio)
-		PRAGMA_ENABLE_DEPRECATION_WARNINGS
 		{
 			bHasAudio = true;
 			break;
@@ -271,9 +259,7 @@ bool FHighlightRecorder::SaveHighlightInBackgroundImpl(const FString& Filename,
 		return false;
 	}
 
-	PRAGMA_DISABLE_DEPRECATION_WARNINGS
 	checkf(Samples[FirstSampleIndex].IsVideoKeyFrame(), TEXT("t %.3f d %.3f"), Samples[FirstSampleIndex].Timestamp.GetTotalSeconds(), Samples[FirstSampleIndex].Duration.GetTotalSeconds());
-	PRAGMA_ENABLE_DEPRECATION_WARNINGS
 
 	if (FirstSampleIndex == Samples.Num())
 	{
@@ -281,22 +267,18 @@ bool FHighlightRecorder::SaveHighlightInBackgroundImpl(const FString& Filename,
 		return false;
 	}
 
-	PRAGMA_DISABLE_DEPRECATION_WARNINGS
 	UE_LOG(HighlightRecorder, Verbose, TEXT("writting %d samples to .mp4, %.3f s, starting from %.3f s, index %d"),
 		Samples.Num() - FirstSampleIndex,
 		(Samples.Last().Timestamp - StartTime + Samples.Last().Duration).GetTotalSeconds(),
 		StartTime.GetTotalSeconds(),
 		FirstSampleIndex);
-	PRAGMA_ENABLE_DEPRECATION_WARNINGS
 
 	// get samples starting from `StartTime` and push them into Mp4Writer
 	for (int Idx = FirstSampleIndex; Idx != Samples.Num(); ++Idx)
 	{
-		PRAGMA_DISABLE_DEPRECATION_WARNINGS
 		AVEncoder::FMediaPacket& Sample = Samples[Idx];
 		Sample.Timestamp = Sample.Timestamp - StartTime;
 		if (!Mp4Writer->Write(Sample, (Sample.Type==AVEncoder::EPacketType::Audio) ? AudioStreamIndex : VideoStreamIndex))
-		PRAGMA_ENABLE_DEPRECATION_WARNINGS
 		{
 			return false;
 		}
@@ -310,7 +292,6 @@ bool FHighlightRecorder::SaveHighlightInBackgroundImpl(const FString& Filename,
 	return true;
 }
 
-
 bool FHighlightRecorder::InitialiseMp4Writer(const FString& FullFilename, bool bHasAudio)
 {
 	// create target directory if it does not exist
@@ -335,9 +316,7 @@ bool FHighlightRecorder::InitialiseMp4Writer(const FString& FullFilename, bool b
 
 	if (bHasAudio)
 	{
-		PRAGMA_DISABLE_DEPRECATION_WARNINGS
 		TPair<FString, AVEncoder::FAudioConfig> AudioConfig = TPair<FString, AVEncoder::FAudioConfig>( FGameplayMediaEncoder::Get()->GetAudioConfig().Codec, FGameplayMediaEncoder::Get()->GetAudioConfig());		
-		PRAGMA_ENABLE_DEPRECATION_WARNINGS
 		if (AudioConfig.Key == "")
 		{
 			UE_LOG(HighlightRecorder, Error, TEXT("Could not get audio config"));
@@ -355,9 +334,7 @@ bool FHighlightRecorder::InitialiseMp4Writer(const FString& FullFilename, bool b
 		}
 	}
 
-	PRAGMA_DISABLE_DEPRECATION_WARNINGS
 	TPair<FString, AVEncoder::FVideoConfig> VideoConfig = TPair<FString, AVEncoder::FVideoConfig>( FGameplayMediaEncoder::Get()->GetVideoConfig().Codec, FGameplayMediaEncoder::Get()->GetVideoConfig());
-	PRAGMA_ENABLE_DEPRECATION_WARNINGS
 	if (VideoConfig.Key == "")
 	{
 		UE_LOG(HighlightRecorder, Error, TEXT("Could not get video config"));
@@ -383,7 +360,6 @@ bool FHighlightRecorder::InitialiseMp4Writer(const FString& FullFilename, bool b
 }
 
 // finds index and timestamp of the first sample that should be written to .mp4
-PRAGMA_DISABLE_DEPRECATION_WARNINGS
 bool FHighlightRecorder::GetSavingStart(const TArray<AVEncoder::FMediaPacket>& Samples, FTimespan MaxDuration, int& StartIndex, FTimespan& StartTime) const
 // the first sample in .mp4 file should have timestamp 0 and all other timestamps should be relative to the
 // first one
@@ -435,4 +411,3 @@ bool FHighlightRecorder::GetSavingStart(const TArray<AVEncoder::FMediaPacket>& S
 
 	return true;
 }
-PRAGMA_ENABLE_DEPRECATION_WARNINGS
\ No newline at end of file
diff --git a/Engine/Source/Runtime/Windows/WindowsPlatformFeatures/Private/HighlightRecorder.h b/Engine/Source/Runtime/Windows/WindowsPlatformFeatures/Private/HighlightRecorder.h
--- a/Engine/Source/Runtime/Windows/WindowsPlatformFeatures/Private/HighlightRecorder.h
+++ b/Engine/Source/Runtime/Windows/WindowsPlatformFeatures/Private/HighlightRecorder.h
@@ -45,9 +45,7 @@ private:
 	bool SaveHighlightInBackground(const FString& Filename, double MaxDurationSecs);
 	bool SaveHighlightInBackgroundImpl(const FString& Filename, double MaxDurationSecs);
 	bool InitialiseMp4Writer(const FString& Filename, bool bHasAudio);
-	PRAGMA_DISABLE_DEPRECATION_WARNINGS
 	bool GetSavingStart(const TArray<AVEncoder::FMediaPacket>& Samples, FTimespan MaxDuration, int& OutStartIndex, FTimespan& OutStartTime) const;
-	PRAGMA_ENABLE_DEPRECATION_WARNINGS
 
 	// takes into account if we've been paused and shifts current time back to compensate paused state
 	// so all timestamps are continuous even over paused pieces
@@ -56,9 +54,7 @@ private:
 	//
 	// IGameplayMediaEncoderListener implementation
 	//
-	PRAGMA_DISABLE_DEPRECATION_WARNINGS
 	void OnMediaSample(const AVEncoder::FMediaPacket& Sample) override;
-	PRAGMA_ENABLE_DEPRECATION_WARNINGS
 
 private:
 	TAtomic<EState> State{ EState::Stopped };
diff --git a/Engine/Source/Runtime/Windows/WindowsPlatformFeatures/Private/WindowsVideoRecordingSystem.cpp b/Engine/Source/Runtime/Windows/WindowsPlatformFeatures/Private/WindowsVideoRecordingSystem.cpp
--- a/Engine/Source/Runtime/Windows/WindowsPlatformFeatures/Private/WindowsVideoRecordingSystem.cpp
+++ b/Engine/Source/Runtime/Windows/WindowsPlatformFeatures/Private/WindowsVideoRecordingSystem.cpp
@@ -155,23 +155,19 @@ private:
 
 
 
-PRAGMA_DISABLE_DEPRECATION_WARNINGS
 FWindowsVideoRecordingSystem::FWindowsVideoRecordingSystem()
 {
 	UE_LOG(WindowsVideoRecordingSystem, Verbose, TEXT("%s"), __FUNCTIONW__);
 
 	EnableRecording(true);
 }
-PRAGMA_ENABLE_DEPRECATION_WARNINGS
 
-PRAGMA_DISABLE_DEPRECATION_WARNINGS
 FWindowsVideoRecordingSystem::~FWindowsVideoRecordingSystem()
 {
 	UE_LOG(WindowsVideoRecordingSystem, Verbose, TEXT("%s"), __FUNCTIONW__);
 
 	EnableRecording(false);
 }
-PRAGMA_ENABLE_DEPRECATION_WARNINGS
 
 void FWindowsVideoRecordingSystem::EnableRecording(bool bEnableRecording)
 {
@@ -409,11 +405,8 @@ void FWindowsVideoRecordingSystem::FinalizeCallbackOnGameThread(bool bSaved, boo
 
 	if (bBroadcast)
 	{
-		PRAGMA_DISABLE_DEPRECATION_WARNINGS
 		OnVideoRecordingFinalized.Broadcast(bSaved, Path);
-		PRAGMA_ENABLE_DEPRECATION_WARNINGS
 	}
-	
 }
 
 EVideoRecordingState FWindowsVideoRecordingSystem::GetRecordingState() const
diff --git a/Engine/Source/Runtime/Windows/WindowsPlatformFeatures/Private/WindowsVideoRecordingSystem.h b/Engine/Source/Runtime/Windows/WindowsPlatformFeatures/Private/WindowsVideoRecordingSystem.h
--- a/Engine/Source/Runtime/Windows/WindowsPlatformFeatures/Private/WindowsVideoRecordingSystem.h
+++ b/Engine/Source/Runtime/Windows/WindowsPlatformFeatures/Private/WindowsVideoRecordingSystem.h
@@ -18,9 +18,7 @@ DECLARE_LOG_CATEGORY_EXTERN(WindowsVideoRecordingSystem, VeryVerbose, VeryVerbos
 
 class FHighlightRecorder;
 
-PRAGMA_DISABLE_DEPRECATION_WARNINGS
 class FWindowsVideoRecordingSystem : public IVideoRecordingSystem
-PRAGMA_ENABLE_DEPRECATION_WARNINGS
 {
 public:
 	FWindowsVideoRecordingSystem();
diff --git a/Engine/Source/Runtime/Windows/WindowsPlatformFeatures/Private/WmfMp4Writer.cpp b/Engine/Source/Runtime/Windows/WindowsPlatformFeatures/Private/WmfMp4Writer.cpp
--- a/Engine/Source/Runtime/Windows/WindowsPlatformFeatures/Private/WmfMp4Writer.cpp
+++ b/Engine/Source/Runtime/Windows/WindowsPlatformFeatures/Private/WmfMp4Writer.cpp
@@ -21,7 +21,6 @@ bool FWmfMp4Writer::Initialize(const TCHAR* Filename)
 	return true;
 }
 
-PRAGMA_DISABLE_DEPRECATION_WARNINGS
 TOptional<DWORD> FWmfMp4Writer::CreateAudioStream(const FString& Codec, const AVEncoder::FAudioConfig& Config)
 {
 	GUID Format;
@@ -59,9 +58,7 @@ TOptional<DWORD> FWmfMp4Writer::CreateAudioStream(const FString& Codec, const AV
 	CHECK_HR_DEFAULT(Writer->SetInputMediaType(StreamIndex, MediaType, nullptr));
 	return TOptional<DWORD>(StreamIndex);
 }
-PRAGMA_ENABLE_DEPRECATION_WARNINGS
 
-PRAGMA_DISABLE_DEPRECATION_WARNINGS
 TOptional<DWORD> FWmfMp4Writer::CreateVideoStream(const FString& Codec, const AVEncoder::FVideoConfig& Config)
 {
 	GUID Format;
@@ -97,7 +94,6 @@ TOptional<DWORD> FWmfMp4Writer::CreateVideoStream(const FString& Codec, const AV
 	CHECK_HR_DEFAULT(Writer->SetInputMediaType(StreamIndex, MediaType, nullptr));
 	return TOptional<DWORD>(StreamIndex);
 }
-PRAGMA_ENABLE_DEPRECATION_WARNINGS
 
 bool FWmfMp4Writer::Start()
 {
@@ -105,7 +101,6 @@ bool FWmfMp4Writer::Start()
 	return true;
 }
 
-PRAGMA_DISABLE_DEPRECATION_WARNINGS
 bool FWmfMp4Writer::Write(const AVEncoder::FMediaPacket& InSample, DWORD StreamIndex)
 {
 	AVEncoder::FIMFSampleWrapper Sample { InSample.Type };
@@ -133,7 +128,6 @@ bool FWmfMp4Writer::Write(const AVEncoder::FMediaPacket& InSample, DWORD StreamI
 
 	return true;
 }
-PRAGMA_ENABLE_DEPRECATION_WARNINGS
 
 bool FWmfMp4Writer::Finalize()
 {
diff --git a/Engine/Source/Runtime/Windows/WindowsPlatformFeatures/Private/WmfMp4Writer.h b/Engine/Source/Runtime/Windows/WindowsPlatformFeatures/Private/WmfMp4Writer.h
--- a/Engine/Source/Runtime/Windows/WindowsPlatformFeatures/Private/WmfMp4Writer.h
+++ b/Engine/Source/Runtime/Windows/WindowsPlatformFeatures/Private/WmfMp4Writer.h
@@ -17,7 +17,6 @@ public:
 	/**
 	 * Create an audio stream and return the its index on success
 	 */
-	PRAGMA_DISABLE_DEPRECATION_WARNINGS
 	TOptional<DWORD> CreateAudioStream(const FString& Codec, const AVEncoder::FAudioConfig& Config);
 
 	/**
@@ -25,10 +24,8 @@ public:
 	 */
 	TOptional<DWORD> CreateVideoStream(const FString& Codec, const AVEncoder::FVideoConfig& Config);
 
-	bool Write(const AVEncoder::FMediaPacket& InSample, DWORD StreamIndex);
-	PRAGMA_ENABLE_DEPRECATION_WARNINGS
-
 	bool Start();
+	bool Write(const AVEncoder::FMediaPacket& InSample, DWORD StreamIndex);
 	bool Finalize();
 
 private:
diff --git a/Engine/Source/Runtime/Windows/WindowsPlatformFeatures/Private/WmfRingBuffer.cpp b/Engine/Source/Runtime/Windows/WindowsPlatformFeatures/Private/WmfRingBuffer.cpp
--- a/Engine/Source/Runtime/Windows/WindowsPlatformFeatures/Private/WmfRingBuffer.cpp
+++ b/Engine/Source/Runtime/Windows/WindowsPlatformFeatures/Private/WmfRingBuffer.cpp
@@ -6,8 +6,6 @@
 DECLARE_LOG_CATEGORY_EXTERN(WmfRingBuffer, Log, VeryVerbose);
 DEFINE_LOG_CATEGORY(WmfRingBuffer);
 
-
-PRAGMA_DISABLE_DEPRECATION_WARNINGS
 void FWmfRingBuffer::Push(AVEncoder::FMediaPacket&& Sample)
 {
 	check(MaxDuration != 0);
@@ -66,7 +64,6 @@ void FWmfRingBuffer::Push(AVEncoder::FMediaPacket&& Sample)
 		}
 	}
 }
-PRAGMA_ENABLE_DEPRECATION_WARNINGS
 
 void FWmfRingBuffer::Cleanup()
 // removes a/v samples grouped by the oldest key-frame period
@@ -75,18 +72,14 @@ void FWmfRingBuffer::Cleanup()
 
 	check(Samples.Num() != 0);
 
-	PRAGMA_DISABLE_DEPRECATION_WARNINGS
 	check(Samples[0].IsVideoKeyFrame());
-	PRAGMA_ENABLE_DEPRECATION_WARNINGS
 
 	// find the second key-frame
 	int i = 1; // skip the first sample which is a keyframe
 	bool bFound = false;
 	for (; i != Samples.Num(); ++i)
 	{
-		PRAGMA_DISABLE_DEPRECATION_WARNINGS
 		if (Samples[i].IsVideoKeyFrame())
-		PRAGMA_ENABLE_DEPRECATION_WARNINGS
 		{
 			bFound = true;
 			break;
@@ -96,18 +89,14 @@ void FWmfRingBuffer::Cleanup()
 
 	// remove key-frame period
 	Samples.RemoveAt(0, i, EAllowShrinking::No);
-	PRAGMA_DISABLE_DEPRECATION_WARNINGS
 	check(Samples[0].IsVideoKeyFrame());
-	PRAGMA_ENABLE_DEPRECATION_WARNINGS
 
-	PRAGMA_DISABLE_DEPRECATION_WARNINGS
 	UE_LOG(WmfRingBuffer, VeryVerbose, TEXT("%d samples, %.3f s, %.3f - %.3f:%.3f"),
 		Samples.Num(),
 		GetDuration().GetTotalSeconds(),
 		Samples[0].Timestamp.GetTotalSeconds(),
 		Samples.Last().Timestamp.GetTotalSeconds(),
 		Samples.Last().Duration.GetTotalSeconds());
-	PRAGMA_ENABLE_DEPRECATION_WARNINGS
 }
 
 void FWmfRingBuffer::PauseCleanup(bool bPause)
@@ -124,12 +113,9 @@ FTimespan FWmfRingBuffer::GetDuration() const
 
 	// A duration of a video is not "LastFrameCaptureTime - FirstFrameCaptureTime" , but 
 	// (LastFrameCaptureTime - FirstFrameCaptureTime) + LastFrameDuration.
-	PRAGMA_DISABLE_DEPRECATION_WARNINGS
 	return Samples.Last().Timestamp + Samples.Last().Duration - Samples[0].Timestamp;
-	PRAGMA_ENABLE_DEPRECATION_WARNINGS
 }
 
-PRAGMA_DISABLE_DEPRECATION_WARNINGS
 TArray<AVEncoder::FMediaPacket> FWmfRingBuffer::GetCopy()
 {
 	FScopeLock Lock(&Mutex);
@@ -142,7 +128,6 @@ TArray<AVEncoder::FMediaPacket> FWmfRingBuffer::GetCopy()
 
 	return Copy;
 }
-PRAGMA_ENABLE_DEPRECATION_WARNINGS
 
 void FWmfRingBuffer::Reset()
 {
diff --git a/Engine/Source/Runtime/Windows/WindowsPlatformFeatures/Private/WmfRingBuffer.h b/Engine/Source/Runtime/Windows/WindowsPlatformFeatures/Private/WmfRingBuffer.h
--- a/Engine/Source/Runtime/Windows/WindowsPlatformFeatures/Private/WmfRingBuffer.h
+++ b/Engine/Source/Runtime/Windows/WindowsPlatformFeatures/Private/WmfRingBuffer.h
@@ -22,13 +22,11 @@ public:
 
 	FTimespan GetDuration() const;
 
-	void PauseCleanup(bool bPause);
-
-	PRAGMA_DISABLE_DEPRECATION_WARNINGS
 	void Push(AVEncoder::FMediaPacket&& Sample);
 
+	void PauseCleanup(bool bPause);
+
 	TArray<AVEncoder::FMediaPacket> GetCopy();
-	PRAGMA_ENABLE_DEPRECATION_WARNINGS
 
 	void Reset();
 
@@ -37,11 +35,8 @@ private:
 
 private:
 	FTimespan MaxDuration = 0;
+	TArray<AVEncoder::FMediaPacket> Samples;
 	FCriticalSection Mutex;
 	FThreadSafeBool bCleanupPaused = false;
-
-	PRAGMA_DISABLE_DEPRECATION_WARNINGS
-	TArray<AVEncoder::FMediaPacket> Samples;
-	PRAGMA_ENABLE_DEPRECATION_WARNINGS
 };
 
diff --git a/Engine/Source/Runtime/Windows/WindowsPlatformFeatures/Public/WindowsPlatformFeatures.h b/Engine/Source/Runtime/Windows/WindowsPlatformFeatures/Public/WindowsPlatformFeatures.h
--- a/Engine/Source/Runtime/Windows/WindowsPlatformFeatures/Public/WindowsPlatformFeatures.h
+++ b/Engine/Source/Runtime/Windows/WindowsPlatformFeatures/Public/WindowsPlatformFeatures.h
@@ -15,13 +15,9 @@ public:
 	WINDOWSPLATFORMFEATURES_API virtual IVideoRecordingSystem* GetVideoRecordingSystem() override;
 	WINDOWSPLATFORMFEATURES_API virtual class ISaveGameSystem* GetSaveGameSystem() override;
 
-	WINDOWSPLATFORMFEATURES_API void RegisterVideoRecordingSystem(TSharedPtr<IVideoRecordingSystem> VideoRecordingSystem);
-
 private:
 	WINDOWSPLATFORMFEATURES_API void StartupModule() override;
 
-private:
-	TSharedPtr<IVideoRecordingSystem> VideoRecordingSystem;
 };
 
 
