Fix: build health issues with unbalanced pragma's

diff --git a/Engine/Source/Runtime/AVEncoder/Private/Encoders/WmfAudioEncoder.cpp b/Engine/Source/Runtime/AVEncoder/Private/Encoders/WmfAudioEncoder.cpp
--- a/Engine/Source/Runtime/AVEncoder/Private/Encoders/WmfAudioEncoder.cpp
+++ b/Engine/Source/Runtime/AVEncoder/Private/Encoders/WmfAudioEncoder.cpp
@@ -200,7 +200,7 @@ bool FWmfAudioEncoder::SetInputType()
 	PRAGMA_DISABLE_DEPRECATION_WARNINGS
 	CHECK_HR(MediaType->SetUINT32(MF_MT_AUDIO_SAMPLES_PER_SECOND, Config.Samplerate));
 	CHECK_HR(MediaType->SetUINT32(MF_MT_AUDIO_NUM_CHANNELS, Config.NumChannels));
-	PRAGMA_DISABLE_DEPRECATION_WARNINGS
+	PRAGMA_ENABLE_DEPRECATION_WARNINGS
 
 	CHECK_HR(Encoder->SetInputType(0, MediaType, 0));
 
@@ -216,9 +216,12 @@ bool FWmfAudioEncoder::SetOutputType()
 	{
 		CHECK_HR(OutputType->SetGUID(MF_MT_SUBTYPE, MFAudioFormat_AAC));
 		CHECK_HR(OutputType->SetUINT32(MF_MT_AUDIO_BITS_PER_SAMPLE, 16)); // the only value supported
+		
+		PRAGMA_DISABLE_DEPRECATION_WARNINGS
 		CHECK_HR(OutputType->SetUINT32(MF_MT_AUDIO_SAMPLES_PER_SECOND, Config.Samplerate));
 		CHECK_HR(OutputType->SetUINT32(MF_MT_AUDIO_NUM_CHANNELS, Config.NumChannels));
 		CHECK_HR(OutputType->SetUINT32(MF_MT_AUDIO_AVG_BYTES_PER_SECOND, Config.Bitrate / 8));
+		PRAGMA_ENABLE_DEPRECATION_WARNINGS
 	}
 	else
 	{
@@ -252,13 +255,17 @@ void FWmfAudioEncoder::Shutdown()
 	// Nothing to do
 }
 
+
+PRAGMA_DISABLE_DEPRECATION_WARNINGS
 void FWmfAudioEncoder::Encode(const FAudioFrame& Frame)
+PRAGMA_ENABLE_DEPRECATION_WARNINGS
 {
 	UE_LOG(LogAVEncoder, Verbose, TEXT("Audio input: time %.3f, duration %.3f, %d samples"), Frame.Timestamp.GetTotalSeconds(), Frame.Duration.GetTotalSeconds(), Frame.Data.GetNumSamples());
 
 	// Convert float audio data to PCM
 	PCM16 = Frame.Data;
 
+	PRAGMA_DISABLE_DEPRECATION_WARNINGS
 	FIMFSampleWrapper InputSample = CreateInputSample(
 		reinterpret_cast<const uint8*>(PCM16.GetData()),
 		PCM16.GetNumSamples()*sizeof(*PCM16.GetData()),
@@ -269,11 +276,13 @@ void FWmfAudioEncoder::Encode(const FAudioFrame& Frame)
 	{
 		return;
 	}
+	PRAGMA_ENABLE_DEPRECATION_WARNINGS
 
 	CHECK_HR_VOID(Encoder->ProcessInput(0, InputSample.GetSample(), 0));
 
 	while (true)
 	{
+		PRAGMA_DISABLE_DEPRECATION_WARNINGS
 		FIMFSampleWrapper OutputSample = GetOutputSample();
 		if (!OutputSample.IsValid())
 		{
@@ -293,13 +302,17 @@ void FWmfAudioEncoder::Encode(const FAudioFrame& Frame)
 			Packet.Data.Append(Data.GetData(), Data.Num());
 			return false; // Terminate, because we only want 1 buffer
 		});
+		PRAGMA_ENABLE_DEPRECATION_WARNINGS
 
 		OnEncodedAudioFrame(Packet);
 	}
 }
 
+PRAGMA_DISABLE_DEPRECATION_WARNINGS
 FIMFSampleWrapper FWmfAudioEncoder::CreateInputSample(const uint8* SampleData, uint32 Size, FTimespan Timestamp, FTimespan Duration)
+PRAGMA_ENABLE_DEPRECATION_WARNINGS
 {
+	PRAGMA_DISABLE_DEPRECATION_WARNINGS
 	FIMFSampleWrapper Sample = { EPacketType::Audio };
 
 	if (!Sample.CreateSample())
@@ -323,11 +336,15 @@ FIMFSampleWrapper FWmfAudioEncoder::CreateInputSample(const uint8* SampleData, u
 	Sample.SetTime(Timestamp);
 	Sample.SetDuration(Duration);
 
+	PRAGMA_ENABLE_DEPRECATION_WARNINGS
 	return Sample;
 }
 
+PRAGMA_DISABLE_DEPRECATION_WARNINGS
 FIMFSampleWrapper FWmfAudioEncoder::GetOutputSample()
+PRAGMA_ENABLE_DEPRECATION_WARNINGS
 {
+	PRAGMA_DISABLE_DEPRECATION_WARNINGS
 	bool bFlag1 = OutputStreamInfo.dwFlags&MFT_OUTPUT_STREAM_PROVIDES_SAMPLES;
 	bool bFlag2 = OutputStreamInfo.dwFlags&MFT_OUTPUT_STREAM_CAN_PROVIDE_SAMPLES;
 
@@ -379,7 +396,7 @@ FIMFSampleWrapper FWmfAudioEncoder::GetOutputSample()
 	{
 		return {};
 	}
-
+	PRAGMA_ENABLE_DEPRECATION_WARNINGS
 }
 
 FAudioConfig FWmfAudioEncoder::GetConfig() const
@@ -416,8 +433,11 @@ TUniquePtr<FAudioEncoder> FWmfAudioEncoderFactory::CreateEncoder(const FString&
 {
 	if (Codec == "aac")
 	{
+		PRAGMA_DISABLE_DEPRECATION_WARNINGS
 		return TUniquePtr<FAudioEncoder>(new FWmfAudioEncoder(FWmfAudioEncoder::ECodecType::AAC));
+		PRAGMA_ENABLE_DEPRECATION_WARNINGS
 	}
+	
 	else
 	{
 		UE_LOG(LogAVEncoder, Error, TEXT("FWmfAudioEncoderFactory doesn't support the %s codec"), *Codec);
