Fix: Compilation errors in VideoEncoderFactory.cpp

diff --git a/Engine/Source/Runtime/AVEncoder/Private/Encoders/WmfAudioEncoder.cpp b/Engine/Source/Runtime/AVEncoder/Private/Encoders/WmfAudioEncoder.cpp
--- a/Engine/Source/Runtime/AVEncoder/Private/Encoders/WmfAudioEncoder.cpp
+++ b/Engine/Source/Runtime/AVEncoder/Private/Encoders/WmfAudioEncoder.cpp
@@ -446,13 +446,10 @@ TArray<FString> FWmfAudioEncoderFactory::GetSupportedCodecs() const
 
 PRAGMA_DISABLE_DEPRECATION_WARNINGS
 TUniquePtr<FAudioEncoder> FWmfAudioEncoderFactory::CreateEncoder(const FString& Codec)
-PRAGMA_ENABLE_DEPRECATION_WARNINGS
 {
 	if (Codec == "aac")
 	{
-		PRAGMA_DISABLE_DEPRECATION_WARNINGS
 		return TUniquePtr<FAudioEncoder>(new FWmfAudioEncoder(FWmfAudioEncoder::ECodecType::AAC));
-		PRAGMA_ENABLE_DEPRECATION_WARNINGS
 	}
 	
 	else
@@ -461,6 +458,7 @@ PRAGMA_ENABLE_DEPRECATION_WARNINGS
 		return nullptr;
 	}
 }
+PRAGMA_ENABLE_DEPRECATION_WARNINGS
 
 }
 
diff --git a/Engine/Source/Runtime/AVEncoder/Private/VideoEncoderFactory.cpp b/Engine/Source/Runtime/AVEncoder/Private/VideoEncoderFactory.cpp
--- a/Engine/Source/Runtime/AVEncoder/Private/VideoEncoderFactory.cpp
+++ b/Engine/Source/Runtime/AVEncoder/Private/VideoEncoderFactory.cpp
@@ -117,17 +117,12 @@ PRAGMA_ENABLE_DEPRECATION_WARNINGS
 
 PRAGMA_DISABLE_DEPRECATION_WARNINGS
 TUniquePtr<FVideoEncoder> FVideoEncoderFactory::Create(uint32 InID, const FVideoEncoder::FLayerConfig& config)
-PRAGMA_ENABLE_DEPRECATION_WARNINGS
 {
 	// HACK (M84FIX) create encoder without a ready FVideoEncoderInput
-	PRAGMA_DISABLE_DEPRECATION_WARNINGS
 	TUniquePtr<FVideoEncoder>	Result;
-	PRAGMA_ENABLE_DEPRECATION_WARNINGS
 	for (int32 Index = 0; Index < AvailableEncoders.Num(); ++Index)
 	{
-		PRAGMA_DISABLE_DEPRECATION_WARNINGS
 		if (AvailableEncoders[Index].ID == InID)
-		PRAGMA_ENABLE_DEPRECATION_WARNINGS
 		{
 			Result = CreateEncoders[Index]();
 
@@ -135,13 +130,9 @@ PRAGMA_ENABLE_DEPRECATION_WARNINGS
 
 			if (RHIType == ERHIInterfaceType::D3D11)
 			{	
-				PRAGMA_DISABLE_DEPRECATION_WARNINGS
 				TSharedRef<FVideoEncoderInput> Input = FVideoEncoderInput::CreateForD3D11(GDynamicRHI->RHIGetNativeDevice(), true, IsRHIDeviceAMD()).ToSharedRef();
-				PRAGMA_ENABLE_DEPRECATION_WARNINGS
 				
-				PRAGMA_DISABLE_DEPRECATION_WARNINGS
 				if (Result && !Result->Setup(Input, config))
-				PRAGMA_ENABLE_DEPRECATION_WARNINGS
 				{
 					Result.Reset();
 				}
@@ -149,13 +140,9 @@ PRAGMA_ENABLE_DEPRECATION_WARNINGS
 			}
 			else if (RHIType == ERHIInterfaceType::D3D12)
 			{				
-				PRAGMA_DISABLE_DEPRECATION_WARNINGS
 				TSharedRef<FVideoEncoderInput> Input = FVideoEncoderInput::CreateForD3D12(GDynamicRHI->RHIGetNativeDevice(), true, IsRHIDeviceNVIDIA()).ToSharedRef();
-				PRAGMA_ENABLE_DEPRECATION_WARNINGS
 				
-				PRAGMA_DISABLE_DEPRECATION_WARNINGS
 				if (Result && !Result->Setup(Input, config))
-				PRAGMA_ENABLE_DEPRECATION_WARNINGS
 				{
 					Result.Reset();
 				}
@@ -164,19 +151,13 @@ PRAGMA_ENABLE_DEPRECATION_WARNINGS
 #if PLATFORM_DESKTOP && !PLATFORM_APPLE
 			else if (RHIType == ERHIInterfaceType::Vulkan)
 			{
-				PRAGMA_DISABLE_DEPRECATION_WARNINGS
 				AVEncoder::FVulkanDataStruct VulkanData = {	static_cast<VkInstance>(GDynamicRHI->RHIGetNativeInstance()), 
 															static_cast<VkPhysicalDevice>(GDynamicRHI->RHIGetNativePhysicalDevice()), 
 															static_cast<VkDevice>(GDynamicRHI->RHIGetNativeDevice())};
-				PRAGMA_ENABLE_DEPRECATION_WARNINGS
 
-				PRAGMA_DISABLE_DEPRECATION_WARNINGS
 				TSharedRef<FVideoEncoderInput> Input = FVideoEncoderInput::CreateForVulkan( &VulkanData, true).ToSharedRef();
-				PRAGMA_ENABLE_DEPRECATION_WARNINGS
 				
-				PRAGMA_DISABLE_DEPRECATION_WARNINGS
 				if (Result && !Result->Setup(Input, config))
-				PRAGMA_ENABLE_DEPRECATION_WARNINGS
 				{
 					Result.Reset();
 				}
@@ -187,6 +168,8 @@ PRAGMA_ENABLE_DEPRECATION_WARNINGS
 	}
 	return Result;
 }
+PRAGMA_ENABLE_DEPRECATION_WARNINGS
+
 
 PRAGMA_DISABLE_DEPRECATION_WARNINGS
 TUniquePtr<FVideoEncoder> FVideoEncoderFactory::Create(uint32 InID, TSharedPtr<FVideoEncoderInput> InInput, const FVideoEncoder::FLayerConfig& config)
