Fixing single header compiles.

diff --git a/Engine/Plugins/Experimental/Mutable/Source/MutableRuntime/Private/MuR/OpImageLayer.h b/Engine/Plugins/Experimental/Mutable/Source/MutableRuntime/Private/MuR/OpImageLayer.h
--- a/Engine/Plugins/Experimental/Mutable/Source/MutableRuntime/Private/MuR/OpImageLayer.h
+++ b/Engine/Plugins/Experimental/Mutable/Source/MutableRuntime/Private/MuR/OpImageLayer.h
@@ -8,6 +8,7 @@
 #include "Templates/UnrealTemplate.h"
 
 #include "MuR/OpImageBlend.h"
+#include "MuR/MutableRuntimeModule.h"
 
 namespace mu
 {
diff --git a/Engine/Source/Runtime/AVEncoder/Public/AudioTypes.h b/Engine/Source/Runtime/AVEncoder/Public/AudioTypes.h
--- a/Engine/Source/Runtime/AVEncoder/Public/AudioTypes.h
+++ b/Engine/Source/Runtime/AVEncoder/Public/AudioTypes.h
@@ -47,7 +47,9 @@ namespace LibRtc
 		/** The sequence of audio samples */
 		TArrayView<const FAudioChannelSample> AudioSamples;
 		/** The format of audio samples */
+PRAGMA_DISABLE_DEPRECATION_WARNINGS
 		FAudioStreamFormat AudioStreamFormat;
+PRAGMA_ENABLE_DEPRECATION_WARNINGS
 	};
 
 	struct UE_DEPRECATED(5.4, "AVEncoder has been deprecated. Please use the AVCodecs plugin family instead.") FAudioCodecConfig
@@ -65,7 +67,9 @@ namespace LibRtc
 	 * @param AudioUnencodedSampleBuffer - The sample buffer which is being encoded.
 	 * @param [out] OutEncodedAudioByteBuffer - The client side pointer of the byte buffer with encoded information.
 	 */
+PRAGMA_DISABLE_DEPRECATION_WARNINGS
 	using OnAudioEncodeCallback = TFunction<void(const FAudioSampleBuffer&, const FAudioByteBuffer& AudioBytes)>;
+PRAGMA_ENABLE_DEPRECATION_WARNINGS
 
 	/**
 	 * The callback is called to decode the byte-by-byte buffer to the sample buffer format.
@@ -76,11 +80,15 @@ namespace LibRtc
 	 * @param AudioUndecodedByteBuffer - The byte buffer which is being decoded.
 	 * @param [out] OutDecodedAudioSampleBuffer - The client side pointer of the sample buffer with decoded information.
 	 */
+PRAGMA_DISABLE_DEPRECATION_WARNINGS
 	using OnAudioDecodeCallback = TFunction<void(const FAudioByteBuffer& AudioBytes, const FAudioSampleBuffer&)>;
+PRAGMA_ENABLE_DEPRECATION_WARNINGS
 
 	struct UE_DEPRECATED(5.4, "AVEncoder has been deprecated. Please use the AVCodecs plugin family instead.") FAudioCodec
 	{
+PRAGMA_DISABLE_DEPRECATION_WARNINGS
 		FAudioCodecConfig AudioCodecConfig;
+PRAGMA_ENABLE_DEPRECATION_WARNINGS
 		OnAudioEncodeCallback OnAudioEncodeCallback;
 		OnAudioDecodeCallback OnAudioDecodeCallback;
 	};
@@ -108,6 +116,8 @@ namespace LibRtc
 
 	struct UE_DEPRECATED(5.4, "AVEncoder has been deprecated. Please use the AVCodecs plugin family instead.") FAudioMuteInfo
 	{
+PRAGMA_DISABLE_DEPRECATION_WARNINGS
 		TOptional<EAudioMuteReason> Reason;
+PRAGMA_ENABLE_DEPRECATION_WARNINGS
 	};
 }
diff --git a/Engine/Source/Runtime/AVEncoder/Public/Configuration.h b/Engine/Source/Runtime/AVEncoder/Public/Configuration.h
--- a/Engine/Source/Runtime/AVEncoder/Public/Configuration.h
+++ b/Engine/Source/Runtime/AVEncoder/Public/Configuration.h
@@ -19,13 +19,17 @@ namespace LibRtc
             static constexpr const TCHAR* kDefaultStunServerUrl = TEXT("stun:stun.l.google.com:19302");
 
             FString StunServerUrl = kDefaultStunServerUrl;
+PRAGMA_DISABLE_DEPRECATION_WARNINGS
             TOptional<FAudioCodec> CustomAudioCodec;
-        };
+PRAGMA_ENABLE_DEPRECATION_WARNINGS
+		};
 
         static FData& Get()
         {
+PRAGMA_DISABLE_DEPRECATION_WARNINGS
             static FData Instance;
             return Instance;
+PRAGMA_ENABLE_DEPRECATION_WARNINGS
         }
     };
 }
diff --git a/Engine/Source/Runtime/Engine/Public/PrimitiveSceneDesc.h b/Engine/Source/Runtime/Engine/Public/PrimitiveSceneDesc.h
--- a/Engine/Source/Runtime/Engine/Public/PrimitiveSceneDesc.h
+++ b/Engine/Source/Runtime/Engine/Public/PrimitiveSceneDesc.h
@@ -5,6 +5,7 @@
 #include "PrimitiveComponentId.h"
 #include "PrimitiveSceneInfoData.h"
 #include "Engine/EngineTypes.h"
+#include "Math/BoxSphereBounds.h"
 
 class FPrimitiveSceneProxy;
 struct FPrimitiveSceneProxyDesc;
