Fix server build creation for cosmetic capture

#rnx
#tests built locally and in Horde, run test with the new build

diff --git a/Engine/Source/Runtime/AVEncoder/AVEncoder.Build.cs b/Engine/Source/Runtime/AVEncoder/AVEncoder.Build.cs
--- a/Engine/Source/Runtime/AVEncoder/AVEncoder.Build.cs
+++ b/Engine/Source/Runtime/AVEncoder/AVEncoder.Build.cs
@@ -44,9 +44,12 @@ public class AVEncoder : ModuleRules
 			});
 
 			AddEngineThirdPartyPrivateStaticDependencies(Target, "DX11", "DX12");
-			
-			PublicDelayLoadDLLs.Add("mfplat.dll");
-			PublicDelayLoadDLLs.Add("mfuuid.dll");
+
+			PublicSystemLibraries.AddRange(new string[] {
+				"mfplat.lib",
+				"mfuuid.lib",
+			});
+
 			PublicDelayLoadDLLs.Add("Mfreadwrite.dll");
 		}
 	}
